Move FEditorModeInfo and shared mode defines to EditorFramework module.

#rnx

diff --git a/Engine/Plugins/2D/Paper2D/Source/Paper2DEditor/Paper2DEditor.Build.cs b/Engine/Plugins/2D/Paper2D/Source/Paper2DEditor/Paper2DEditor.Build.cs
--- a/Engine/Plugins/2D/Paper2D/Source/Paper2DEditor/Paper2DEditor.Build.cs
+++ b/Engine/Plugins/2D/Paper2D/Source/Paper2DEditor/Paper2DEditor.Build.cs
@@ -17,6 +17,7 @@ public class Paper2DEditor : ModuleRules
 				"SlateCore",
 				"Engine",
 				"InputCore",
+				"EditorFramework",
 				"UnrealEd", // for AssetEditorSubsystem
 				"KismetWidgets",
 				"Kismet",  // for FWorkflowCentricApplication
diff --git a/Engine/Plugins/Editor/GeometryMode/Source/BspMode/BspMode.Build.cs b/Engine/Plugins/Editor/GeometryMode/Source/BspMode/BspMode.Build.cs
--- a/Engine/Plugins/Editor/GeometryMode/Source/BspMode/BspMode.Build.cs
+++ b/Engine/Plugins/Editor/GeometryMode/Source/BspMode/BspMode.Build.cs
@@ -15,6 +15,7 @@ public class BspMode : ModuleRules
 				"Slate",
 				"SlateCore",
                 "EditorStyle",
+				"EditorFramework",
 				"UnrealEd",
 				"LevelEditor",
 			}
diff --git a/Engine/Plugins/Editor/GeometryMode/Source/TextureAlignMode/TextureAlignMode.Build.cs b/Engine/Plugins/Editor/GeometryMode/Source/TextureAlignMode/TextureAlignMode.Build.cs
--- a/Engine/Plugins/Editor/GeometryMode/Source/TextureAlignMode/TextureAlignMode.Build.cs
+++ b/Engine/Plugins/Editor/GeometryMode/Source/TextureAlignMode/TextureAlignMode.Build.cs
@@ -13,6 +13,7 @@ public class TextureAlignMode : ModuleRules
 				"Engine",
 				"SlateCore",
 				"Slate",
+				"EditorFramework",
 				"UnrealEd",
 				"RenderCore",
 				"LevelEditor",
diff --git a/Engine/Plugins/Experimental/ChaosEditor/Source/FractureEditor/FractureEditor.Build.cs b/Engine/Plugins/Experimental/ChaosEditor/Source/FractureEditor/FractureEditor.Build.cs
--- a/Engine/Plugins/Experimental/ChaosEditor/Source/FractureEditor/FractureEditor.Build.cs
+++ b/Engine/Plugins/Experimental/ChaosEditor/Source/FractureEditor/FractureEditor.Build.cs
@@ -42,6 +42,7 @@ public class FractureEditor : ModuleRules
 				"Slate",
 				"SlateCore",
 				"InputCore",
+				"EditorFramework",
 				"UnrealEd",
 				"LevelEditor",
                 "GeometryCollectionCore",
diff --git a/Engine/Plugins/Experimental/ControlRig/Source/ControlRigEditor/ControlRigEditor.Build.cs b/Engine/Plugins/Experimental/ControlRig/Source/ControlRigEditor/ControlRigEditor.Build.cs
--- a/Engine/Plugins/Experimental/ControlRig/Source/ControlRigEditor/ControlRigEditor.Build.cs
+++ b/Engine/Plugins/Experimental/ControlRig/Source/ControlRigEditor/ControlRigEditor.Build.cs
@@ -33,7 +33,8 @@ namespace UnrealBuildTool.Rules
                     "SlateCore",
                     "InputCore",
                     "Engine",
-                    "UnrealEd",
+					"EditorFramework",
+					"UnrealEd",
                     "KismetCompiler",
                     "BlueprintGraph",
                     "ControlRig",
diff --git a/Engine/Plugins/Experimental/GeometryCollectionPlugin/Source/GeometryCollectionEditor/GeometryCollectionEditor.Build.cs b/Engine/Plugins/Experimental/GeometryCollectionPlugin/Source/GeometryCollectionEditor/GeometryCollectionEditor.Build.cs
--- a/Engine/Plugins/Experimental/GeometryCollectionPlugin/Source/GeometryCollectionEditor/GeometryCollectionEditor.Build.cs
+++ b/Engine/Plugins/Experimental/GeometryCollectionPlugin/Source/GeometryCollectionEditor/GeometryCollectionEditor.Build.cs
@@ -17,7 +17,8 @@ namespace UnrealBuildTool.Rules
 				    "Slate",
 				    "SlateCore",
 				    "Engine",
-				    "UnrealEd",
+					"EditorFramework",
+					"UnrealEd",
 				    "PropertyEditor",
 				    "RenderCore",
 				    "RHI",
diff --git a/Engine/Plugins/Experimental/GizmoEdMode/Source/GizmoEdMode.Build.cs b/Engine/Plugins/Experimental/GizmoEdMode/Source/GizmoEdMode.Build.cs
--- a/Engine/Plugins/Experimental/GizmoEdMode/Source/GizmoEdMode.Build.cs
+++ b/Engine/Plugins/Experimental/GizmoEdMode/Source/GizmoEdMode.Build.cs
@@ -36,6 +36,7 @@ public class GizmoEdMode : ModuleRules
 			{
 				"Engine",
 				"InputCore",
+				"EditorFramework",
 				"UnrealEd",
 				"SlateCore",
 				"CoreUObject",
diff --git a/Engine/Plugins/Experimental/ModelingToolsEditorMode/Source/ModelingToolsEditorMode/ModelingToolsEditorMode.Build.cs b/Engine/Plugins/Experimental/ModelingToolsEditorMode/Source/ModelingToolsEditorMode/ModelingToolsEditorMode.Build.cs
--- a/Engine/Plugins/Experimental/ModelingToolsEditorMode/Source/ModelingToolsEditorMode/ModelingToolsEditorMode.Build.cs
+++ b/Engine/Plugins/Experimental/ModelingToolsEditorMode/Source/ModelingToolsEditorMode/ModelingToolsEditorMode.Build.cs
@@ -42,7 +42,8 @@ public class ModelingToolsEditorMode : ModuleRules
                 "SlateCore",
                 "Engine",
                 "InputCore",
-                "UnrealEd",
+				"EditorFramework",
+				"UnrealEd",
                 "ContentBrowser",
                 "LevelEditor",
 				"EditorStyle",
diff --git a/Engine/Plugins/Experimental/SampleToolsEditorMode/Source/SampleToolsEditorMode.Build.cs b/Engine/Plugins/Experimental/SampleToolsEditorMode/Source/SampleToolsEditorMode.Build.cs
--- a/Engine/Plugins/Experimental/SampleToolsEditorMode/Source/SampleToolsEditorMode.Build.cs
+++ b/Engine/Plugins/Experimental/SampleToolsEditorMode/Source/SampleToolsEditorMode.Build.cs
@@ -40,6 +40,7 @@ public class SampleToolsEditorMode : ModuleRules
 				"Slate",
 				"SlateCore",
 				"InputCore",
+				"EditorFramework",
 				"UnrealEd",
                 "ContentBrowser",
                 "LevelEditor",
diff --git a/Engine/Plugins/MeshPainting/Source/MeshPaintEditorMode/MeshPaintEditorMode.Build.cs b/Engine/Plugins/MeshPainting/Source/MeshPaintEditorMode/MeshPaintEditorMode.Build.cs
--- a/Engine/Plugins/MeshPainting/Source/MeshPaintEditorMode/MeshPaintEditorMode.Build.cs
+++ b/Engine/Plugins/MeshPainting/Source/MeshPaintEditorMode/MeshPaintEditorMode.Build.cs
@@ -21,7 +21,8 @@ namespace UnrealBuildTool.Rules
 				    "CoreUObject",
 				    "Engine",
                     "InputCore",
-				    "UnrealEd",
+					"EditorFramework",
+					"UnrealEd",
 					"InteractiveToolsFramework",
 					"EditorInteractiveToolsFramework",
 					"MeshPaintingToolset",
diff --git a/Engine/Source/Developer/AssetTools/AssetTools.Build.cs b/Engine/Source/Developer/AssetTools/AssetTools.Build.cs
--- a/Engine/Source/Developer/AssetTools/AssetTools.Build.cs
+++ b/Engine/Source/Developer/AssetTools/AssetTools.Build.cs
@@ -13,6 +13,7 @@ public class AssetTools : ModuleRules
 				"Core",
 				"CoreUObject",
 				"SlateCore",
+				"EditorFramework",
 				"UnrealEd",
 			}
 		);
diff --git a/Engine/Source/Developer/GameplayDebugger/GameplayDebugger.Build.cs b/Engine/Source/Developer/GameplayDebugger/GameplayDebugger.Build.cs
--- a/Engine/Source/Developer/GameplayDebugger/GameplayDebugger.Build.cs
+++ b/Engine/Source/Developer/GameplayDebugger/GameplayDebugger.Build.cs
@@ -32,7 +32,8 @@ namespace UnrealBuildTool.Rules
                 PrivateDependencyModuleNames.AddRange(
                     new string[] {
                         "EditorStyle",
-                        "UnrealEd",
+						"EditorFramework",
+						"UnrealEd",
                         "LevelEditor",
                         "PropertyEditor",
                     });
diff --git a/Engine/Source/Editor/ActorPickerMode/ActorPickerMode.Build.cs b/Engine/Source/Editor/ActorPickerMode/ActorPickerMode.Build.cs
--- a/Engine/Source/Editor/ActorPickerMode/ActorPickerMode.Build.cs
+++ b/Engine/Source/Editor/ActorPickerMode/ActorPickerMode.Build.cs
@@ -15,6 +15,7 @@ public class ActorPickerMode : ModuleRules
                 "Slate",
                 "SlateCore",
                 "EditorStyle",
+				"EditorFramework",
 				"UnrealEd",
 			}
 		);
diff --git a/Engine/Source/Editor/AnimGraph/AnimGraph.Build.cs b/Engine/Source/Editor/AnimGraph/AnimGraph.Build.cs
--- a/Engine/Source/Editor/AnimGraph/AnimGraph.Build.cs
+++ b/Engine/Source/Editor/AnimGraph/AnimGraph.Build.cs
@@ -30,6 +30,7 @@ public class AnimGraph : ModuleRules
 			new string[] {
 				"InputCore",
 				"SlateCore",
+				"EditorFramework",
 				"UnrealEd",
                 "GraphEditor",
 				"PropertyEditor",
diff --git a/Engine/Source/Editor/ClothPainter/ClothPainter.Build.cs b/Engine/Source/Editor/ClothPainter/ClothPainter.Build.cs
--- a/Engine/Source/Editor/ClothPainter/ClothPainter.Build.cs
+++ b/Engine/Source/Editor/ClothPainter/ClothPainter.Build.cs
@@ -16,7 +16,8 @@ public class ClothPainter : ModuleRules
                 "Slate",
                 "SlateCore",
                 "EditorStyle",
-                "UnrealEd",
+				"EditorFramework",
+				"UnrealEd",
                 "MainFrame",
                 "PropertyEditor",
                 "Kismet",
diff --git a/Engine/Source/Editor/DetailCustomizations/DetailCustomizations.Build.cs b/Engine/Source/Editor/DetailCustomizations/DetailCustomizations.Build.cs
--- a/Engine/Source/Editor/DetailCustomizations/DetailCustomizations.Build.cs
+++ b/Engine/Source/Editor/DetailCustomizations/DetailCustomizations.Build.cs
@@ -22,6 +22,7 @@ public class DetailCustomizations : ModuleRules
 				"Slate",
 				"SlateCore",
 				"EditorStyle",
+				"EditorFramework",
 				"UnrealEd",
 				"EditorWidgets",
 				"KismetWidgets",
diff --git a/Engine/Source/Editor/EditorFramework/EditorFramework.Build.cs b/Engine/Source/Editor/EditorFramework/EditorFramework.Build.cs
new file mode 100644
--- /dev/null
+++ b/Engine/Source/Editor/EditorFramework/EditorFramework.Build.cs
@@ -0,0 +1,48 @@
+// Copyright Epic Games, Inc. All Rights Reserved.
+
+using UnrealBuildTool;
+
+public class EditorFramework : ModuleRules
+{
+	public EditorFramework(ReadOnlyTargetRules Target) : base(Target)
+	{
+		PrivateIncludePaths.AddRange(
+			new string[] {
+			}
+		);
+
+		PrivateIncludePathModuleNames.AddRange(
+			new string[] {
+			}
+		);
+
+		PrivateDependencyModuleNames.AddRange(
+			new string[] {
+				"Core",
+				"CoreUObject",
+				"EditorStyle",
+				"EditorSubsystem",
+				"Engine",
+				"InputCore",
+				"RenderCore",
+				"Slate",
+				"SlateCore",
+			}
+		);
+
+		PublicDependencyModuleNames.AddRange(
+			new string[] {
+			}
+		);
+
+		PublicIncludePathModuleNames.AddRange(
+			new string[] {
+			}
+		);
+
+		DynamicallyLoadedModuleNames.AddRange(
+            new string[] { 
+            }
+		);
+	}
+}
diff --git a/Engine/Source/Editor/EditorFramework/Private/EditorFrameworkModule.cpp b/Engine/Source/Editor/EditorFramework/Private/EditorFrameworkModule.cpp
new file mode 100644
--- /dev/null
+++ b/Engine/Source/Editor/EditorFramework/Private/EditorFrameworkModule.cpp
@@ -0,0 +1,15 @@
+// Copyright Epic Games, Inc. All Rights Reserved.
+
+#include "EditorFrameworkModule.h"
+
+#include "Modules/ModuleManager.h"
+
+void FEditorFrameworkModule::StartupModule()
+{
+}
+
+void FEditorFrameworkModule::ShutdownModule()
+{
+}
+
+IMPLEMENT_MODULE(FEditorFrameworkModule, EditorFramework);
diff --git a/Engine/Source/Editor/EditorFramework/Private/EditorModes.cpp b/Engine/Source/Editor/EditorFramework/Private/EditorModes.cpp
new file mode 100644
--- /dev/null
+++ b/Engine/Source/Editor/EditorFramework/Private/EditorModes.cpp
@@ -0,0 +1,24 @@
+// Copyright Epic Games, Inc. All Rights Reserved.
+
+#include "EditorModes.h"
+#include "EdMode.h"
+
+
+DEFINE_LOG_CATEGORY(LogEditorModes);
+
+// Builtin editor mode constants
+namespace FBuiltinEditorModes
+{
+	const FEditorModeID EM_None = NAME_None;
+	const FEditorModeID EM_Default(TEXT("EM_Default"));
+	const FEditorModeID EM_Placement(TEXT("PLACEMENT"));
+	const FEditorModeID EM_InterpEdit(TEXT("EM_InterpEdit"));
+	const FEditorModeID EM_MeshPaint(TEXT("EM_MeshPaint"));
+	const FEditorModeID EM_Landscape(TEXT("EM_Landscape"));
+	const FEditorModeID EM_Foliage(TEXT("EM_Foliage"));
+	const FEditorModeID EM_Level(TEXT("EM_Level"));
+	const FEditorModeID EM_StreamingLevel(TEXT("EM_StreamingLevel"));
+	const FEditorModeID EM_Physics(TEXT("EM_Physics"));
+	const FEditorModeID EM_ActorPicker(TEXT("EM_ActorPicker"));
+	const FEditorModeID EM_SceneDepthPicker(TEXT("EM_SceneDepthPicker"));
+}
diff --git a/Engine/Source/Editor/EditorFramework/Private/Tools/Modes.cpp b/Engine/Source/Editor/EditorFramework/Private/Tools/Modes.cpp
new file mode 100644
--- /dev/null
+++ b/Engine/Source/Editor/EditorFramework/Private/Tools/Modes.cpp
@@ -0,0 +1,30 @@
+// Copyright Epic Games, Inc. All Rights Reserved.
+
+#include "Tools/Modes.h"
+
+FEditorModeInfo::FEditorModeInfo()
+	: ID(NAME_None)
+	, bVisible(false)
+	, PriorityOrder(MAX_int32)
+{
+}
+
+FEditorModeInfo::FEditorModeInfo(
+	FEditorModeID InID,
+	FText InName,
+	FSlateIcon InIconBrush,
+	bool InIsVisible,
+	int32 InPriorityOrder
+	)
+	: ID(InID)
+	, ToolbarCustomizationName(*(InID.ToString() + TEXT("Toolbar")))
+	, Name(InName)
+	, IconBrush(InIconBrush)
+	, bVisible(InIsVisible)
+	, PriorityOrder(InPriorityOrder)
+{
+	if (!InIconBrush.IsSet())
+	{
+		IconBrush = FSlateIcon(FEditorStyle::GetStyleSetName(), "LevelEditor.EditorModes");
+	}
+}
\ No newline at end of file
diff --git a/Engine/Source/Editor/EditorFramework/Public/EditorFrameworkModule.h b/Engine/Source/Editor/EditorFramework/Public/EditorFrameworkModule.h
new file mode 100644
--- /dev/null
+++ b/Engine/Source/Editor/EditorFramework/Public/EditorFrameworkModule.h
@@ -0,0 +1,12 @@
+// Copyright Epic Games, Inc. All Rights Reserved.
+
+#pragma once
+
+#include "Modules/ModuleInterface.h"
+
+class FEditorFrameworkModule : public IModuleInterface
+{
+public:
+	virtual void StartupModule() override;
+	virtual void ShutdownModule() override;
+};
diff --git a/Engine/Source/Editor/EditorFramework/Public/EditorModes.h b/Engine/Source/Editor/EditorFramework/Public/EditorModes.h
new file mode 100644
--- /dev/null
+++ b/Engine/Source/Editor/EditorFramework/Public/EditorModes.h
@@ -0,0 +1,74 @@
+// Copyright Epic Games, Inc. All Rights Reserved.
+
+#pragma once
+
+#include "CoreMinimal.h"
+#include "MaterialShared.h"
+#include "SceneManagement.h"
+#include "Tools/Modes.h"
+
+EDITORFRAMEWORK_API DECLARE_LOG_CATEGORY_EXTERN(LogEditorModes, Log, All);
+
+// Builtin editor mode constants
+namespace FBuiltinEditorModes
+{
+	/** Gameplay, editor disabled. */
+	EDITORFRAMEWORK_API extern const FEditorModeID EM_None;
+
+	/** Camera movement, actor placement. */
+	EDITORFRAMEWORK_API extern const FEditorModeID EM_Default;
+
+	/** Placement mode */
+	EDITORFRAMEWORK_API extern const FEditorModeID EM_Placement;
+
+	/** Interpolation editing. */
+	EDITORFRAMEWORK_API extern const FEditorModeID EM_InterpEdit;
+
+	/** Mesh paint tool */
+	EDITORFRAMEWORK_API extern const FEditorModeID EM_MeshPaint;
+
+	/** Landscape editing */
+	EDITORFRAMEWORK_API extern const FEditorModeID EM_Landscape;
+
+	/** Foliage painting */
+	EDITORFRAMEWORK_API extern const FEditorModeID EM_Foliage;
+
+	/** Level editing mode */
+	EDITORFRAMEWORK_API extern const FEditorModeID EM_Level;
+
+	/** Streaming level editing mode */
+	EDITORFRAMEWORK_API extern const FEditorModeID EM_StreamingLevel;
+
+	/** Physics manipulation mode ( available only when simulating in viewport )*/
+	EDITORFRAMEWORK_API extern const FEditorModeID EM_Physics;
+
+	/** Actor picker mode, used to interactively pick actors in the viewport */
+	EDITORFRAMEWORK_API extern const FEditorModeID EM_ActorPicker;
+
+	/** Actor picker mode, used to interactively pick actors in the viewport */
+	EDITORFRAMEWORK_API extern const FEditorModeID EM_SceneDepthPicker;
+};
+
+/** Material proxy wrapper that can be created on the game thread and passed on to the render thread. */
+class UNREALED_API FDynamicColoredMaterialRenderProxy : public FDynamicPrimitiveResource, public FColoredMaterialRenderProxy
+{
+public:
+	/** Initialization constructor. */
+	FDynamicColoredMaterialRenderProxy(const FMaterialRenderProxy* InParent,const FLinearColor& InColor)
+	:	FColoredMaterialRenderProxy(InParent,InColor)
+	{
+	}
+	virtual ~FDynamicColoredMaterialRenderProxy()
+	{
+	}
+
+	// FDynamicPrimitiveResource interface.
+	virtual void InitPrimitiveResource()
+	{
+	}
+	virtual void ReleasePrimitiveResource()
+	{
+		delete this;
+	}
+};
+
diff --git a/Engine/Source/Editor/EditorFramework/Public/Tools/Modes.h b/Engine/Source/Editor/EditorFramework/Public/Tools/Modes.h
new file mode 100644
--- /dev/null
+++ b/Engine/Source/Editor/EditorFramework/Public/Tools/Modes.h
@@ -0,0 +1,48 @@
+// Copyright Epic Games, Inc. All Rights Reserved.
+
+#pragma once
+
+#include "Delegates/DelegateCombinations.h"
+#include "Internationalization/Text.h"
+#include "Textures/SlateIcon.h"
+#include "UObject/NameTypes.h"
+
+/** The shorthand identifier used for editor modes */
+typedef FName FEditorModeID;
+
+struct FEditorModeInfo
+{
+	/** Default constructor */
+	EDITORFRAMEWORK_API FEditorModeInfo();
+
+	/** Helper constructor */
+	EDITORFRAMEWORK_API FEditorModeInfo(
+		FEditorModeID InID,
+		FText InName = FText(),
+		FSlateIcon InIconBrush = FSlateIcon(),
+		bool bInIsVisible = false,
+		int32 InPriorityOrder = MAX_int32
+		);
+
+	/** The mode ID */
+	FEditorModeID ID;
+
+	/** Name of the toolbar this mode uses and can be used by external systems to customize that mode toolbar */
+	FName ToolbarCustomizationName;
+
+	/** Name for the editor to display */
+	FText Name;
+
+	/** The mode icon */
+	FSlateIcon IconBrush;
+
+	/** Whether or not the mode should be visible in the mode menu */
+	bool bVisible;
+
+	/** The priority of this mode which will determine its default order and shift+X command assignment */
+	int32 PriorityOrder;
+};
+
+DECLARE_MULTICAST_DELEGATE(FRegisteredModesChangedEvent);
+DECLARE_MULTICAST_DELEGATE_OneParam(FOnModeRegistered, FEditorModeID);
+DECLARE_MULTICAST_DELEGATE_OneParam(FOnModeUnregistered, FEditorModeID);
diff --git a/Engine/Source/Editor/FoliageEdit/FoliageEdit.build.cs b/Engine/Source/Editor/FoliageEdit/FoliageEdit.build.cs
--- a/Engine/Source/Editor/FoliageEdit/FoliageEdit.build.cs
+++ b/Engine/Source/Editor/FoliageEdit/FoliageEdit.build.cs
@@ -13,6 +13,7 @@ public class FoliageEdit : ModuleRules
 				"CoreUObject",
 				"InputCore",
 				"Engine",
+				"EditorFramework",
 				"UnrealEd",
 				"Slate",
 				"SlateCore",
diff --git a/Engine/Source/Editor/LandscapeEditor/LandscapeEditor.Build.cs b/Engine/Source/Editor/LandscapeEditor/LandscapeEditor.Build.cs
--- a/Engine/Source/Editor/LandscapeEditor/LandscapeEditor.Build.cs
+++ b/Engine/Source/Editor/LandscapeEditor/LandscapeEditor.Build.cs
@@ -20,6 +20,7 @@ public class LandscapeEditor : ModuleRules
                 "RenderCore",
                 "RHI",
                 "InputCore",
+				"EditorFramework",
 				"UnrealEd",
 				"PropertyEditor",
 				"ImageWrapper",
diff --git a/Engine/Source/Editor/LevelEditor/LevelEditor.Build.cs b/Engine/Source/Editor/LevelEditor/LevelEditor.Build.cs
--- a/Engine/Source/Editor/LevelEditor/LevelEditor.Build.cs
+++ b/Engine/Source/Editor/LevelEditor/LevelEditor.Build.cs
@@ -48,6 +48,7 @@ public class LevelEditor : ModuleRules
 				"SourceControl",
 				"SourceControlWindows",
 				"StatsViewer",
+				"EditorFramework",
 				"UnrealEd", 
 				"RenderCore",
 				"DeviceProfileServices",
diff --git a/Engine/Source/Editor/MainFrame/MainFrame.Build.cs b/Engine/Source/Editor/MainFrame/MainFrame.Build.cs
--- a/Engine/Source/Editor/MainFrame/MainFrame.Build.cs
+++ b/Engine/Source/Editor/MainFrame/MainFrame.Build.cs
@@ -31,6 +31,7 @@ public class MainFrame : ModuleRules
 				"SourceControlWindows",
 				"TargetPlatform",
 				"DesktopPlatform",
+				"EditorFramework",
 				"UnrealEd",
 				"WorkspaceMenuStructure",
 				"MessageLog",
diff --git a/Engine/Source/Editor/Matinee/Matinee.Build.cs b/Engine/Source/Editor/Matinee/Matinee.Build.cs
--- a/Engine/Source/Editor/Matinee/Matinee.Build.cs
+++ b/Engine/Source/Editor/Matinee/Matinee.Build.cs
@@ -37,6 +37,7 @@ public class Matinee : ModuleRules
                 "InputCore",
 				"LevelEditor",
 				"DistCurveEditor",
+				"EditorFramework",
 				"UnrealEd",
 				"RenderCore",
 				"AssetRegistry",
diff --git a/Engine/Source/Editor/MeshPaintMode/MeshPaintMode.Build.cs b/Engine/Source/Editor/MeshPaintMode/MeshPaintMode.Build.cs
--- a/Engine/Source/Editor/MeshPaintMode/MeshPaintMode.Build.cs
+++ b/Engine/Source/Editor/MeshPaintMode/MeshPaintMode.Build.cs
@@ -32,7 +32,8 @@ public class MeshPaintMode : ModuleRules
                 "Slate",
 				"SlateCore",
                 "EditorStyle",
-                "UnrealEd",
+				"EditorFramework",
+				"UnrealEd",
                 "RawMesh",
                 "SourceControl",
                 "PropertyEditor",
diff --git a/Engine/Source/Editor/Persona/Persona.Build.cs b/Engine/Source/Editor/Persona/Persona.Build.cs
--- a/Engine/Source/Editor/Persona/Persona.Build.cs
+++ b/Engine/Source/Editor/Persona/Persona.Build.cs
@@ -46,8 +46,9 @@ public class Persona : ModuleRules
                 "SlateCore",
 				"ContentBrowserData",
                 "EditorStyle",
-                "Engine", 
-                "UnrealEd", 
+                "Engine",
+				"EditorFramework",
+				"UnrealEd", 
                 "GraphEditor", 
                 "InputCore",
                 "Kismet", 
@@ -63,7 +64,6 @@ public class Persona : ModuleRules
                 "ClothingSystemRuntimeInterface",
                 "ClothingSystemRuntimeCommon",
                 "AnimGraphRuntime",
-                "UnrealEd",
                 "CommonMenuExtensions",
                 "PinnedCommandList",
                 "RenderCore",
diff --git a/Engine/Source/Editor/PhysicsAssetEditor/PhysicsAssetEditor.Build.cs b/Engine/Source/Editor/PhysicsAssetEditor/PhysicsAssetEditor.Build.cs
--- a/Engine/Source/Editor/PhysicsAssetEditor/PhysicsAssetEditor.Build.cs
+++ b/Engine/Source/Editor/PhysicsAssetEditor/PhysicsAssetEditor.Build.cs
@@ -31,6 +31,7 @@ public class PhysicsAssetEditor : ModuleRules
 				"SlateCore",
                 "EditorStyle",
 				"LevelEditor",
+				"EditorFramework",
 				"UnrealEd",
                 "Kismet",
                 "Persona",
diff --git a/Engine/Source/Editor/PlacementMode/PlacementMode.Build.cs b/Engine/Source/Editor/PlacementMode/PlacementMode.Build.cs
--- a/Engine/Source/Editor/PlacementMode/PlacementMode.Build.cs
+++ b/Engine/Source/Editor/PlacementMode/PlacementMode.Build.cs
@@ -16,7 +16,8 @@ public class PlacementMode : ModuleRules
                 "InputCore",
                 "Slate",
 				"SlateCore",
-                "UnrealEd",
+				"EditorFramework",
+				"UnrealEd",
                 "ContentBrowser",
                 "CollectionManager",
                 "LevelEditor",
diff --git a/Engine/Source/Editor/SceneDepthPickerMode/SceneDepthPickerMode.Build.cs b/Engine/Source/Editor/SceneDepthPickerMode/SceneDepthPickerMode.Build.cs
--- a/Engine/Source/Editor/SceneDepthPickerMode/SceneDepthPickerMode.Build.cs
+++ b/Engine/Source/Editor/SceneDepthPickerMode/SceneDepthPickerMode.Build.cs
@@ -22,6 +22,7 @@ public class SceneDepthPickerMode : ModuleRules
                 "Slate",
                 "SlateCore",
                 "EditorStyle",
+				"EditorFramework",
 				"UnrealEd",
 			}
 		);
diff --git a/Engine/Source/Editor/Sequencer/Sequencer.Build.cs b/Engine/Source/Editor/Sequencer/Sequencer.Build.cs
--- a/Engine/Source/Editor/Sequencer/Sequencer.Build.cs
+++ b/Engine/Source/Editor/Sequencer/Sequencer.Build.cs
@@ -36,6 +36,7 @@ public class Sequencer : ModuleRules
 				"SlateCore",
 				"SceneOutliner",
                 "EditorStyle",
+				"EditorFramework",
 				"UnrealEd", 
 				"MovieScene", 
 				"MovieSceneTracks", 
diff --git a/Engine/Source/Editor/UnrealEd/Private/EdMode.cpp b/Engine/Source/Editor/UnrealEd/Private/EdMode.cpp
--- a/Engine/Source/Editor/UnrealEd/Private/EdMode.cpp
+++ b/Engine/Source/Editor/UnrealEd/Private/EdMode.cpp
@@ -1115,3 +1115,16 @@ bool FEdMode::IsSnapRotationEnabled()
 {
 	return GetDefault<ULevelEditorViewportSettings>()->RotGridEnabled;
 }
+
+/*------------------------------------------------------------------------------
+	Default.
+------------------------------------------------------------------------------*/
+
+FEdModeDefault::FEdModeDefault()
+{
+	bDrawGrid = true;
+	bDrawPivot = false;
+	bDrawBaseInfo = false;
+	bDrawWorldBox = false;
+	bDrawKillZ = false;
+}
diff --git a/Engine/Source/Editor/UnrealEd/Private/EditorModeRegistry.cpp b/Engine/Source/Editor/UnrealEd/Private/EditorModeRegistry.cpp
--- a/Engine/Source/Editor/UnrealEd/Private/EditorModeRegistry.cpp
+++ b/Engine/Source/Editor/UnrealEd/Private/EditorModeRegistry.cpp
@@ -41,33 +41,6 @@ TSharedRef<FEdMode> FEditorModeFactory::CreateMode() const
 	return FactoryCallback.Execute();
 }
 
-FEditorModeInfo::FEditorModeInfo()
-	: ID(NAME_None)
-	, bVisible(false)
-	, PriorityOrder(MAX_int32)
-{
-}
-
-FEditorModeInfo::FEditorModeInfo(
-	FEditorModeID InID,
-	FText InName,
-	FSlateIcon InIconBrush,
-	bool InIsVisible,
-	int32 InPriorityOrder
-	)
-	: ID(InID)
-	, ToolbarCustomizationName(*(InID.ToString()+TEXT("Toolbar")))
-	, Name(InName)
-	, IconBrush(InIconBrush)
-	, bVisible(InIsVisible)
-	, PriorityOrder(InPriorityOrder)
-{
-	if (!InIconBrush.IsSet())
-	{
-		IconBrush = FSlateIcon(FEditorStyle::GetStyleSetName(), "LevelEditor.EditorModes");
-	}
-}
-
 void FEditorModeRegistry::Initialize()
 {
 	// Send notifications for any legacy modes that were registered before the asset subsytem started up
diff --git a/Engine/Source/Editor/UnrealEd/Private/EditorModes.cpp b/Engine/Source/Editor/UnrealEd/Private/EditorModes.cpp
deleted file mode 100644
diff --git a/Engine/Source/Editor/UnrealEd/Private/EditorViewportClient.cpp b/Engine/Source/Editor/UnrealEd/Private/EditorViewportClient.cpp
--- a/Engine/Source/Editor/UnrealEd/Private/EditorViewportClient.cpp
+++ b/Engine/Source/Editor/UnrealEd/Private/EditorViewportClient.cpp
@@ -1192,49 +1192,6 @@ FSceneView* FEditorViewportClient::CalcSceneView(FSceneViewFamily* ViewFamily, c
 	return View;
 }
 
-/** Determines if the new MoveCanvas movement should be used
- * @return - true if we should use the new drag canvas movement.  Returns false for combined object-camera movement and marquee selection
- */
-bool FLevelEditorViewportClient::ShouldUseMoveCanvasMovement()
-{
-	const bool LeftMouseButtonDown = Viewport->KeyState(EKeys::LeftMouseButton) ? true : false;
-	const bool MiddleMouseButtonDown = Viewport->KeyState(EKeys::MiddleMouseButton) ? true : false;
-	const bool RightMouseButtonDown = Viewport->KeyState(EKeys::RightMouseButton) ? true : false;
-	const bool bMouseButtonDown = (LeftMouseButtonDown || MiddleMouseButtonDown || RightMouseButtonDown );
-
-	const bool AltDown = IsAltPressed();
-	const bool ShiftDown = IsShiftPressed();
-	const bool ControlDown = IsCtrlPressed();
-
-	//if we're using the new move canvas mode, we're in an ortho viewport, and the mouse is down
-	if (GetDefault<ULevelEditorViewportSettings>()->bPanMovesCanvas && IsOrtho() && bMouseButtonDown)
-	{
-		//MOVING CAMERA
-		if ( !MouseDeltaTracker->UsingDragTool() && AltDown == false && ShiftDown == false && ControlDown == false && (Widget->GetCurrentAxis() == EAxisList::None) && (LeftMouseButtonDown ^ RightMouseButtonDown))
-		{
-			return true;
-		}
-
-		//OBJECT MOVEMENT CODE
-		if ( ( AltDown == false && ShiftDown == false && ( LeftMouseButtonDown ^ RightMouseButtonDown ) ) &&
-			( ( GetWidgetMode() == FWidget::WM_Translate && Widget->GetCurrentAxis() != EAxisList::None ) ||
-			( GetWidgetMode() == FWidget::WM_TranslateRotateZ && Widget->GetCurrentAxis() != EAxisList::ZRotation &&  Widget->GetCurrentAxis() != EAxisList::None ) ||
-			( GetWidgetMode() == FWidget::WM_2D && Widget->GetCurrentAxis() != EAxisList::Rotate2D &&  Widget->GetCurrentAxis() != EAxisList::None ) ) )
-		{
-			return true;
-		}
-
-
-		//ALL other cases hide the mouse
-		return false;
-	}
-	else
-	{
-		//current system - do not show cursor when mouse is down
-		return false;
-	}
-}
-
 void FEditorViewportClient::ReceivedFocus(FViewport* InViewport)
 {
 	// Viewport has changed got to reset the cursor as it could of been left in any state
diff --git a/Engine/Source/Editor/UnrealEd/Private/LevelEditorViewport.cpp b/Engine/Source/Editor/UnrealEd/Private/LevelEditorViewport.cpp
--- a/Engine/Source/Editor/UnrealEd/Private/LevelEditorViewport.cpp
+++ b/Engine/Source/Editor/UnrealEd/Private/LevelEditorViewport.cpp
@@ -4961,6 +4961,49 @@ void FLevelEditorViewportClient::UpdateAudioListener(const FSceneView& View)
 	}
 }
 
+/** Determines if the new MoveCanvas movement should be used
+ * @return - true if we should use the new drag canvas movement.  Returns false for combined object-camera movement and marquee selection
+ */
+bool FLevelEditorViewportClient::ShouldUseMoveCanvasMovement()
+{
+	const bool LeftMouseButtonDown = Viewport->KeyState(EKeys::LeftMouseButton) ? true : false;
+	const bool MiddleMouseButtonDown = Viewport->KeyState(EKeys::MiddleMouseButton) ? true : false;
+	const bool RightMouseButtonDown = Viewport->KeyState(EKeys::RightMouseButton) ? true : false;
+	const bool bMouseButtonDown = (LeftMouseButtonDown || MiddleMouseButtonDown || RightMouseButtonDown );
+
+	const bool AltDown = IsAltPressed();
+	const bool ShiftDown = IsShiftPressed();
+	const bool ControlDown = IsCtrlPressed();
+
+	//if we're using the new move canvas mode, we're in an ortho viewport, and the mouse is down
+	if (GetDefault<ULevelEditorViewportSettings>()->bPanMovesCanvas && IsOrtho() && bMouseButtonDown)
+	{
+		//MOVING CAMERA
+		if ( !MouseDeltaTracker->UsingDragTool() && AltDown == false && ShiftDown == false && ControlDown == false && (Widget->GetCurrentAxis() == EAxisList::None) && (LeftMouseButtonDown ^ RightMouseButtonDown))
+		{
+			return true;
+		}
+
+		//OBJECT MOVEMENT CODE
+		if ( ( AltDown == false && ShiftDown == false && ( LeftMouseButtonDown ^ RightMouseButtonDown ) ) &&
+			( ( GetWidgetMode() == FWidget::WM_Translate && Widget->GetCurrentAxis() != EAxisList::None ) ||
+			( GetWidgetMode() == FWidget::WM_TranslateRotateZ && Widget->GetCurrentAxis() != EAxisList::ZRotation &&  Widget->GetCurrentAxis() != EAxisList::None ) ||
+			( GetWidgetMode() == FWidget::WM_2D && Widget->GetCurrentAxis() != EAxisList::Rotate2D &&  Widget->GetCurrentAxis() != EAxisList::None ) ) )
+		{
+			return true;
+		}
+
+
+		//ALL other cases hide the mouse
+		return false;
+	}
+	else
+	{
+		//current system - do not show cursor when mouse is down
+		return false;
+	}
+}
+
 void FLevelEditorViewportClient::SetupViewForRendering( FSceneViewFamily& ViewFamily, FSceneView& View )
 {
 	FEditorViewportClient::SetupViewForRendering( ViewFamily, View );
diff --git a/Engine/Source/Editor/UnrealEd/Private/Tools/UEditorModes.cpp b/Engine/Source/Editor/UnrealEd/Private/Tools/UEditorModes.cpp
deleted file mode 100644
diff --git a/Engine/Source/Editor/UnrealEd/Public/EditorModes.h b/Engine/Source/Editor/UnrealEd/Public/EditorModes.h
deleted file mode 100644
diff --git a/Engine/Source/Editor/UnrealEd/Public/Tools/Modes.h b/Engine/Source/Editor/UnrealEd/Public/Tools/Modes.h
deleted file mode 100644
diff --git a/Engine/Source/Editor/UnrealEd/UnrealEd.Build.cs b/Engine/Source/Editor/UnrealEd/UnrealEd.Build.cs
--- a/Engine/Source/Editor/UnrealEd/UnrealEd.Build.cs
+++ b/Engine/Source/Editor/UnrealEd/UnrealEd.Build.cs
@@ -86,6 +86,7 @@ public class UnrealEd : ModuleRules
 				"SandboxFile",
 				"Slate",
 				"SlateCore",
+				"EditorFramework",
 				"EditorStyle",
 				"SourceControl",
 				"UnrealEdMessages",
diff --git a/Engine/Source/Editor/VREditor/VREditor.Build.cs b/Engine/Source/Editor/VREditor/VREditor.Build.cs
--- a/Engine/Source/Editor/VREditor/VREditor.Build.cs
+++ b/Engine/Source/Editor/VREditor/VREditor.Build.cs
@@ -19,7 +19,8 @@ namespace UnrealBuildTool.Rules
 				    "Slate",
 					"SlateCore",
                     "EditorStyle",
-				    "UnrealEd",
+					"EditorFramework",
+					"UnrealEd",
 					"UMG",
 					"LevelEditor",
 					"HeadMountedDisplay",
diff --git a/Engine/Source/Editor/WorldBrowser/WorldBrowser.Build.cs b/Engine/Source/Editor/WorldBrowser/WorldBrowser.Build.cs
--- a/Engine/Source/Editor/WorldBrowser/WorldBrowser.Build.cs
+++ b/Engine/Source/Editor/WorldBrowser/WorldBrowser.Build.cs
@@ -35,7 +35,8 @@ public class WorldBrowser : ModuleRules
                 "Slate",
 				"SlateCore",
                 "EditorStyle",
-                "UnrealEd",
+				"EditorFramework",
+				"UnrealEd",
                 "GraphEditor",
                 "LevelEditor",
                 "PropertyEditor",
diff --git a/Engine/Source/Runtime/Launch/Launch.Build.cs b/Engine/Source/Runtime/Launch/Launch.Build.cs
--- a/Engine/Source/Runtime/Launch/Launch.Build.cs
+++ b/Engine/Source/Runtime/Launch/Launch.Build.cs
@@ -185,6 +185,7 @@ public class Launch : ModuleRules
 
 			PrivateDependencyModuleNames.AddRange(new string[] {
 					"SourceControl",
+					"EditorFramework",
 					"UnrealEd",
 					"DesktopPlatform",
 					"PIEPreviewDeviceProfileSelector",
