Move Toolkit Manager to EditorFramework module.

#rnx

diff --git a/Engine/Plugins/Editor/GeometryMode/Source/GeometryMode/GeometryMode.Build.cs b/Engine/Plugins/Editor/GeometryMode/Source/GeometryMode/GeometryMode.Build.cs
--- a/Engine/Plugins/Editor/GeometryMode/Source/GeometryMode/GeometryMode.Build.cs
+++ b/Engine/Plugins/Editor/GeometryMode/Source/GeometryMode/GeometryMode.Build.cs
@@ -15,6 +15,7 @@ public class GeometryMode : ModuleRules
 				"Slate",
 				"SlateCore",
                 "EditorStyle",
+				"EditorFramework",
 				"UnrealEd",
 				"RenderCore",
 				"LevelEditor",
diff --git a/Engine/Plugins/Editor/MeshEditor/Source/MeshEditor/MeshEditor.Build.cs b/Engine/Plugins/Editor/MeshEditor/Source/MeshEditor/MeshEditor.Build.cs
--- a/Engine/Plugins/Editor/MeshEditor/Source/MeshEditor/MeshEditor.Build.cs
+++ b/Engine/Plugins/Editor/MeshEditor/Source/MeshEditor/MeshEditor.Build.cs
@@ -15,7 +15,8 @@ namespace UnrealBuildTool.Rules
 				    "Slate",
 					"SlateCore",
                     "EditorStyle",
-				    "UnrealEd",
+					"EditorFramework",
+					"UnrealEd",
                     "SceneOutliner",
                     "RenderCore",
 					"EditableMesh",
diff --git a/Engine/Plugins/Enterprise/DatasmithCADImporter/Source/DatasmithCoreTechExtension/DatasmithCoreTechExtension.Build.cs b/Engine/Plugins/Enterprise/DatasmithCADImporter/Source/DatasmithCoreTechExtension/DatasmithCoreTechExtension.Build.cs
--- a/Engine/Plugins/Enterprise/DatasmithCADImporter/Source/DatasmithCoreTechExtension/DatasmithCoreTechExtension.Build.cs
+++ b/Engine/Plugins/Enterprise/DatasmithCADImporter/Source/DatasmithCoreTechExtension/DatasmithCoreTechExtension.Build.cs
@@ -30,6 +30,7 @@ namespace UnrealBuildTool.Rules
 					"Slate",
 					"SlateCore",
 					"StaticMeshDescription",
+					"EditorFramework",
 					"UnrealEd",
 					"CADLibrary",
 				}
diff --git a/Engine/Plugins/Runtime/Metasound/Source/MetasoundEditor/MetasoundEditor.Build.cs b/Engine/Plugins/Runtime/Metasound/Source/MetasoundEditor/MetasoundEditor.Build.cs
--- a/Engine/Plugins/Runtime/Metasound/Source/MetasoundEditor/MetasoundEditor.Build.cs
+++ b/Engine/Plugins/Runtime/Metasound/Source/MetasoundEditor/MetasoundEditor.Build.cs
@@ -43,6 +43,7 @@ namespace UnrealBuildTool.Rules
 					"RenderCore",
 					"Slate",
 					"SlateCore",
+					"EditorFramework",
 					"UnrealEd",
 				}
 			);
diff --git a/Engine/Plugins/Runtime/Steam/SteamAudio/Source/SteamAudioEditor/SteamAudioEditor.Build.cs b/Engine/Plugins/Runtime/Steam/SteamAudio/Source/SteamAudioEditor/SteamAudioEditor.Build.cs
--- a/Engine/Plugins/Runtime/Steam/SteamAudio/Source/SteamAudioEditor/SteamAudioEditor.Build.cs
+++ b/Engine/Plugins/Runtime/Steam/SteamAudio/Source/SteamAudioEditor/SteamAudioEditor.Build.cs
@@ -27,6 +27,7 @@ namespace UnrealBuildTool.Rules
 					"CoreUObject",
 					"Engine",
 					"InputCore",
+					"EditorFramework",
 					"UnrealEd",
 					"LevelEditor",
 					"EditorStyle",
diff --git a/Engine/Plugins/VirtualProduction/Takes/Source/TakeRecorder/TakeRecorder.Build.cs b/Engine/Plugins/VirtualProduction/Takes/Source/TakeRecorder/TakeRecorder.Build.cs
--- a/Engine/Plugins/VirtualProduction/Takes/Source/TakeRecorder/TakeRecorder.Build.cs
+++ b/Engine/Plugins/VirtualProduction/Takes/Source/TakeRecorder/TakeRecorder.Build.cs
@@ -30,6 +30,7 @@ public class TakeRecorder : ModuleRules
 				"Slate",
 				"SlateCore",
 				"UMG",
+				"EditorFramework",
 				"UnrealEd",
 				"WorkspaceMenuStructure",
 				"Analytics",
diff --git a/Engine/Source/Editor/AudioEditor/AudioEditor.Build.cs b/Engine/Source/Editor/AudioEditor/AudioEditor.Build.cs
--- a/Engine/Source/Editor/AudioEditor/AudioEditor.Build.cs
+++ b/Engine/Source/Editor/AudioEditor/AudioEditor.Build.cs
@@ -39,6 +39,7 @@ public class AudioEditor : ModuleRules
 				"AudioMixer",
 				"InputCore",
 				"Engine",
+				"EditorFramework",
 				"UnrealEd",
 				"Slate",
 				"SlateCore",
diff --git a/Engine/Source/Editor/EditorFramework/Private/Toolkits/ToolkitManager.cpp b/Engine/Source/Editor/EditorFramework/Private/Toolkits/ToolkitManager.cpp
new file mode 100644
--- /dev/null
+++ b/Engine/Source/Editor/EditorFramework/Private/Toolkits/ToolkitManager.cpp
@@ -0,0 +1,112 @@
+// Copyright Epic Games, Inc. All Rights Reserved.
+
+#include "Toolkits/ToolkitManager.h"
+#include "Layout/Margin.h"
+#include "Widgets/DeclarativeSyntaxSupport.h"
+#include "Framework/Docking/TabManager.h"
+#include "Toolkits/IToolkit.h"
+#include "Toolkits/IToolkitHost.h"
+#include "Widgets/Docking/SDockTab.h"
+
+static TSharedRef<SDockTab> SpawnStandaloneToolkitHost( const FSpawnTabArgs& Args )
+{
+	return
+		SNew(SDockTab)
+		.TabRole(ETabRole::MajorTab)
+		.ContentPadding(FMargin(0));
+}
+
+static TSharedPtr<SDockTab> NeverReuse(const FTabId&)
+{
+	return TSharedPtr<SDockTab>();
+}
+
+FToolkitManager& FToolkitManager::Get()
+{
+	static TSharedRef< FToolkitManager > ToolkitManagerInstance( new FToolkitManager() );
+	return *ToolkitManagerInstance;
+}
+
+FToolkitManager::FToolkitManager()
+{
+	FGlobalTabmanager::Get()->RegisterTabSpawner( "StandaloneToolkit", FOnSpawnTab::CreateStatic( &SpawnStandaloneToolkitHost ) )
+		.SetReuseTabMethod( FOnFindTabToReuse::CreateStatic( &NeverReuse ) );
+}
+
+
+void FToolkitManager::RegisterNewToolkit( const TSharedRef< IToolkit > NewToolkit )
+{
+	// Add to our list
+	check( !Toolkits.Contains( NewToolkit ) );
+	Toolkits.Add( NewToolkit );
+
+	// Tell the host about this tool kit
+	const TSharedPtr< IToolkitHost > ToolkitHost = NewToolkit->GetToolkitHost();
+	if( ToolkitHost.IsValid() )
+	{
+		ToolkitHost->OnToolkitHostingStarted( NewToolkit );
+	}
+}
+
+
+void FToolkitManager::CloseToolkit( const TSharedRef< IToolkit > ClosingToolkit )
+{
+	// NOTE: This function is called when a user closes a toolkit interactively while the toolkit's host is still
+	//       around.  If the host itself is closed (such as the SLevelEditor goes away during shutdown), then this
+	//       function will not be called.  The toolkit will instead be cleaned up on OnToolkitHostDestroyed, below.
+
+	// Tell the host about this tool kit
+	if( ClosingToolkit->IsHosted() )
+	{
+		const TSharedPtr< IToolkitHost > ToolkitHost = ClosingToolkit->GetToolkitHost();
+		if( ToolkitHost.IsValid() )
+		{
+			ToolkitHost->OnToolkitHostingFinished( ClosingToolkit );
+		}
+	}
+
+	// Remove from our list
+	Toolkits.Remove( ClosingToolkit );
+}
+
+
+
+void FToolkitManager::OnToolkitHostDestroyed( IToolkitHost* HostBeingDestroyed )
+{
+	for( auto ToolkitIndex = 0; ToolkitIndex < Toolkits.Num(); ++ToolkitIndex )
+	{
+		TSharedPtr< IToolkit > Toolkit = Toolkits[ ToolkitIndex ];
+		check( Toolkit.IsValid() );
+
+		if( !Toolkit->IsHosted() || TSharedPtr< IToolkitHost >( Toolkit->GetToolkitHost() ).Get() == HostBeingDestroyed )
+		{
+			// Found a toolkit that is hosted by the host that is going away.  We'll destroy this toolkit now.
+			// NOTE: In this case, OnToolkitHostingFinished() is not called on the host (since it's probably in
+			//       the middle of being destructed!)
+			Toolkits.RemoveAt( ToolkitIndex-- );
+
+			// Continue iterating, so we can cleanup any references to toolkits whose host is no
+			// longer valid.
+		}
+	}
+}
+
+
+TSharedPtr< IToolkit > FToolkitManager::FindEditorForAsset( const UObject* Asset )
+{
+	for( auto ToolkitIndex = 0; ToolkitIndex < Toolkits.Num(); ++ToolkitIndex )
+	{
+		TSharedPtr< IToolkit > Toolkit = Toolkits[ ToolkitIndex ];
+		check( Toolkit.IsValid() );
+
+		if( Toolkit->IsAssetEditor() )
+		{
+			if (Toolkit->GetObjectsCurrentlyBeingEdited()->FindByKey(Asset) != NULL )
+			{
+				return Toolkit;
+			}
+		}
+	}
+
+	return NULL;
+}
diff --git a/Engine/Source/Editor/EditorFramework/Public/Toolkits/IToolkit.h b/Engine/Source/Editor/EditorFramework/Public/Toolkits/IToolkit.h
new file mode 100644
--- /dev/null
+++ b/Engine/Source/Editor/EditorFramework/Public/Toolkits/IToolkit.h
@@ -0,0 +1,151 @@
+// Copyright Epic Games, Inc. All Rights Reserved.
+
+#pragma once
+
+#include "CoreMinimal.h"
+#include "Framework/Docking/WorkspaceItem.h"
+#include "Framework/Docking/TabManager.h"
+
+class FEdMode;
+class UEdMode;
+class IToolkitHost;
+class FEditorModeTools;
+
+/** Some toolkits can be spawned as either standalone tools or within an existing level editing UI */
+namespace EToolkitMode
+{
+	enum Type
+	{
+		/** Stand-alone asset editing "app" */
+		Standalone,
+
+		/** World-centric asset editor, with an interface that sits alongside the level editor */
+		WorldCentric,
+	};
+}
+
+
+namespace EToolkitTabSpot
+{
+	enum Type
+	{
+		// A good place to put a details tab or property tree, or other tabs that display information about 
+		// the current object the user has selected
+		Details,
+
+		// This is for document windows, such as a graph editor or a profiling results panel
+		Document,
+
+		// Put navigation-related panels here, such as a list of editable sub-objects
+		Navigation,
+
+		// Tool bar panels are usually displayed at the top of the tool in a narrow strip
+		ToolBar,
+
+		// Placement panels host tools for creating new things.  This is usually displayed in the top left or center left of a window
+		Placement,
+
+		// Preview viewports
+		Viewport,
+
+		// Area below the level editor. Useful for things that should fill the entire width of the level editor
+		BelowLevelEditor,
+	};
+}
+
+
+/**
+ * Interface for editor toolkits (asset editors and mode tools)
+ */
+class IToolkit
+{
+public:
+
+	/** Register tabs that this toolkit can spawn with the TabManager */
+	virtual void RegisterTabSpawners(const TSharedRef<FTabManager>& TabManager) = 0;
+
+	/** Unregister tabs that this toolkit can spawn */	
+	virtual void UnregisterTabSpawners(const TSharedRef<FTabManager>& TabManager) = 0;
+
+public:
+
+	/** Returns the invariant name of this toolkit type */
+	virtual FName GetToolkitFName() const = 0;
+
+	/** Returns the invariant name of this toolkit type, used for establishing the context for help, documentation and tutorials */
+	virtual FName GetToolkitContextFName() const = 0;
+
+	/** Returns the localized name of this toolkit type (typically just "<ClassName> editor") */
+	virtual FText GetBaseToolkitName() const = 0;
+
+	/** Returns the localized name of this toolkit */
+	virtual FText GetToolkitName() const = 0;
+
+	/** Returns the localized tab suffix */
+	virtual FText GetTabSuffix() const = 0;
+
+	/** Returns the localized tooltip text of this toolkit */
+	virtual FText GetToolkitToolTipText() const = 0;
+
+	/** Returns the workspace menu category of this toolkit */
+	virtual TSharedRef<FWorkspaceItem> GetWorkspaceMenuCategory() const = 0;
+
+	/** Returns the localize prefix string to use for tab labels in world-centric mode. */
+	virtual FString GetWorldCentricTabPrefix() const = 0;
+
+	/** Returns true if this toolkit is used to edit assets (even if it's not necessarily editing one right now.) */
+	virtual bool IsAssetEditor() const = 0;
+
+	/** For asset editor toolkits, returns the UObjects for the assets currently being edited */
+	virtual const TArray< UObject* >* GetObjectsCurrentlyBeingEdited() const = 0;
+
+	/** @return Returns the color and opacity to use for the color that appears behind the tab text for this toolkit's tab in world-centric mode. */
+	virtual FLinearColor GetWorldCentricTabColorScale() const = 0;
+
+	/** @return Returns true if this toolkit is currently hosted.  All toolkits are hosted except during a shutdown situation. */
+	virtual bool IsHosted() const = 0;
+
+	/** @return Returns the toolkit host for this toolkit */
+	virtual const TSharedRef<IToolkitHost> GetToolkitHost() const = 0;
+
+	/**
+	 * Processes any UI commands which are activated by the specified event
+	 *
+	 * @param InKeyEvent	The event to check
+	 *
+	 * @return true if an action was processed
+	 */
+	virtual bool ProcessCommandBindings(const FKeyEvent& InKeyEvent) const = 0;
+
+	/** Call this function to bring all of this toolkit's tabs to the foreground in their respective stacks.  Also causes the toolkit's host window to be foregrounded, too! */
+	virtual void BringToolkitToFront() = 0;
+
+	/* Gets the mode manager for this toolkit */
+	virtual FEditorModeTools& GetEditorModeManager() const = 0;
+
+	/** @returns the editor mode this toolkit is used for, or null if not relevant. */
+	virtual FEdMode* GetEditorMode() const = 0;
+
+	/** @returns the scriptable editor mode this toolkit is used for, or null if not relevant. */
+	virtual UEdMode* GetScriptableEditorMode() const = 0;
+
+	/** @return the editor mode display name if there is an editor mode associated with this toolkit.  An empty FText is returned otherwise */
+	virtual FText GetEditorModeDisplayName() const = 0;
+
+	/** @return the editor mode icon if there is an editor mode associated with this toolkit.  An empty icon is returned otherwise */
+	virtual FSlateIcon GetEditorModeIcon() const = 0;
+
+	/** @returns the inline content that this toolkit returns if it is an editor mode */
+	virtual TSharedPtr<SWidget> GetInlineContent() const = 0;
+
+	/** Returns if this is a IBlueprintEditor derivation */
+	virtual bool IsBlueprintEditor() const = 0;
+
+	/* Called when a toolkit has been brought to the 'front' */
+	virtual void ToolkitBroughtToFront() {};
+
+	/* Returns true if this is the default generic asset editor used by multiple classes */
+	virtual bool IsSimpleAssetEditor() const { return false; }
+};
+
+
diff --git a/Engine/Source/Editor/EditorFramework/Public/Toolkits/IToolkitHost.h b/Engine/Source/Editor/EditorFramework/Public/Toolkits/IToolkitHost.h
new file mode 100644
--- /dev/null
+++ b/Engine/Source/Editor/EditorFramework/Public/Toolkits/IToolkitHost.h
@@ -0,0 +1,40 @@
+// Copyright Epic Games, Inc. All Rights Reserved.
+
+#pragma once
+
+#include "CoreMinimal.h"
+#include "Toolkits/IToolkit.h"
+
+class FEditorModeTools;
+
+/**
+ * Base interface class for toolkit hosts
+ */
+class IToolkitHost
+{
+
+public:
+
+	/** Gets a widget that can be used to parent a modal window or pop-up to.  You shouldn't be using this widget for
+	    anything other than parenting, as the type of widget and behavior/lifespan is completely up to the host. */
+	virtual TSharedRef< class SWidget > GetParentWidget() = 0;
+
+	/** Brings this toolkit host's window (and tab, if it has one), to the front */
+	virtual void BringToFront() = 0;
+
+	/** Access the toolkit host's tab manager */
+	virtual TSharedPtr< class FTabManager > GetTabManager() const = 0;
+
+	/** Called when a toolkit is opened within this host */
+	virtual void OnToolkitHostingStarted( const TSharedRef< class IToolkit >& Toolkit ) = 0;
+
+	/** Called when a toolkit is no longer being hosted within this host */
+	virtual void OnToolkitHostingFinished( const TSharedRef< class IToolkit >& Toolkit ) = 0;
+
+	/** @return For world-centric toolkit hosts, gets the UWorld associated with this host */
+	virtual class UWorld* GetWorld() const = 0;
+
+	/** Returns the mode manager for this toolkit host. For standalone toolkits */
+	virtual FEditorModeTools& GetEditorModeManager() const = 0;
+};
+
diff --git a/Engine/Source/Editor/EditorFramework/Public/Toolkits/ToolkitManager.h b/Engine/Source/Editor/EditorFramework/Public/Toolkits/ToolkitManager.h
new file mode 100644
--- /dev/null
+++ b/Engine/Source/Editor/EditorFramework/Public/Toolkits/ToolkitManager.h
@@ -0,0 +1,44 @@
+// Copyright Epic Games, Inc. All Rights Reserved.
+
+#pragma once
+
+#include "CoreMinimal.h"
+
+class IToolkit;
+class IToolkitHost;
+
+/**
+ * Singleton that managers instances of editor toolkits
+ */
+class FToolkitManager
+{
+
+public:
+	/** Get the singleton instance of the asset editor manager */
+	EDITORFRAMEWORK_API static FToolkitManager& Get();
+
+	/** Call this to register a newly-created toolkit. */
+	EDITORFRAMEWORK_API void RegisterNewToolkit( const TSharedRef< IToolkit > NewToolkit );
+
+	/** Call this to close an existing toolkit */
+	EDITORFRAMEWORK_API void CloseToolkit( const TSharedRef< IToolkit > ClosingToolkit );
+
+	/** Called by a toolkit host itself right before it goes away, so that we can make sure the toolkits are destroyed too */
+	EDITORFRAMEWORK_API void OnToolkitHostDestroyed( IToolkitHost* HostBeingDestroyed );
+
+	/** Tries to find an open asset editor that's editing the specified asset, and returns the toolkit for that asset editor */
+	EDITORFRAMEWORK_API TSharedPtr< IToolkit > FindEditorForAsset( const UObject* Asset );
+
+private:
+	// Buried constructor since the asset editor manager is a singleton
+	FToolkitManager();
+
+
+private:
+
+	/** List of all currently open toolkits */
+	TArray< TSharedPtr< IToolkit > > Toolkits;
+};
+
+
+
diff --git a/Engine/Source/Editor/IntroTutorials/IntroTutorials.Build.cs b/Engine/Source/Editor/IntroTutorials/IntroTutorials.Build.cs
--- a/Engine/Source/Editor/IntroTutorials/IntroTutorials.Build.cs
+++ b/Engine/Source/Editor/IntroTutorials/IntroTutorials.Build.cs
@@ -30,7 +30,8 @@ namespace UnrealBuildTool.Rules
 			PrivateDependencyModuleNames.AddRange(
 				new string[] {
                     "AppFramework",
-                    "UnrealEd",
+					"EditorFramework",
+					"UnrealEd",
                     "Kismet",
                     "PlacementMode",
 					"SlateCore",
diff --git a/Engine/Source/Editor/Kismet/Kismet.Build.cs b/Engine/Source/Editor/Kismet/Kismet.Build.cs
--- a/Engine/Source/Editor/Kismet/Kismet.Build.cs
+++ b/Engine/Source/Editor/Kismet/Kismet.Build.cs
@@ -36,6 +36,7 @@ public class Kismet : ModuleRules
 				"Json",
 				"Merge",
 				"MessageLog",
+				"EditorFramework",
 				"UnrealEd",
 				"GraphEditor",
 				"KismetWidgets",
diff --git a/Engine/Source/Editor/MaterialEditor/MaterialEditor.Build.cs b/Engine/Source/Editor/MaterialEditor/MaterialEditor.Build.cs
--- a/Engine/Source/Editor/MaterialEditor/MaterialEditor.Build.cs
+++ b/Engine/Source/Editor/MaterialEditor/MaterialEditor.Build.cs
@@ -37,6 +37,7 @@ public class MaterialEditor : ModuleRules
 				"RHI",
                 "MaterialUtilities",
                 "PropertyEditor",
+				"EditorFramework",
 				"UnrealEd",
 				"GraphEditor",
                 "AdvancedPreviewScene",
diff --git a/Engine/Source/Editor/MeshPaint/MeshPaint.Build.cs b/Engine/Source/Editor/MeshPaint/MeshPaint.Build.cs
--- a/Engine/Source/Editor/MeshPaint/MeshPaint.Build.cs
+++ b/Engine/Source/Editor/MeshPaint/MeshPaint.Build.cs
@@ -26,7 +26,8 @@ public class MeshPaint : ModuleRules
                 "Slate",
 				"SlateCore",
                 "EditorStyle",
-                "UnrealEd",
+				"EditorFramework",
+				"UnrealEd",
                 "MeshDescription",
 				"StaticMeshDescription",
                 "SourceControl",
diff --git a/Engine/Source/Editor/StaticMeshEditor/StaticMeshEditor.Build.cs b/Engine/Source/Editor/StaticMeshEditor/StaticMeshEditor.Build.cs
--- a/Engine/Source/Editor/StaticMeshEditor/StaticMeshEditor.Build.cs
+++ b/Engine/Source/Editor/StaticMeshEditor/StaticMeshEditor.Build.cs
@@ -30,6 +30,7 @@ public class StaticMeshEditor : ModuleRules
                 "EditorStyle",
 				"RenderCore",
 				"RHI",
+				"EditorFramework",
 				"UnrealEd",
 				"TargetPlatform",
 				"RawMesh",
diff --git a/Engine/Source/Editor/UnrealEd/Private/Toolkits/ToolkitManager.cpp b/Engine/Source/Editor/UnrealEd/Private/Toolkits/ToolkitManager.cpp
deleted file mode 100644
diff --git a/Engine/Source/Editor/UnrealEd/Public/Toolkits/IToolkit.h b/Engine/Source/Editor/UnrealEd/Public/Toolkits/IToolkit.h
deleted file mode 100644
diff --git a/Engine/Source/Editor/UnrealEd/Public/Toolkits/IToolkitHost.h b/Engine/Source/Editor/UnrealEd/Public/Toolkits/IToolkitHost.h
deleted file mode 100644
diff --git a/Engine/Source/Editor/UnrealEd/Public/Toolkits/ToolkitManager.h b/Engine/Source/Editor/UnrealEd/Public/Toolkits/ToolkitManager.h
deleted file mode 100644
