Move Viewport Tab content and Editor Layout to EditorFramework.

#rnx

diff --git a/Engine/Source/Editor/EditorFramework/Private/EditorViewportLayout.cpp b/Engine/Source/Editor/EditorFramework/Private/EditorViewportLayout.cpp
new file mode 100644
--- /dev/null
+++ b/Engine/Source/Editor/EditorFramework/Private/EditorViewportLayout.cpp
@@ -0,0 +1,137 @@
+// Copyright Epic Games, Inc. All Rights Reserved.
+
+#include "EditorViewportLayout.h"
+#include "Layout/Margin.h"
+#include "Widgets/DeclarativeSyntaxSupport.h"
+#include "Widgets/SCompoundWidget.h"
+#include "Widgets/SOverlay.h"
+#include "Framework/Docking/LayoutService.h"
+#include "Misc/ConfigCacheIni.h"
+#include "Modules/ModuleManager.h"
+#include "Layout/WidgetPath.h"
+#include "Framework/Application/SlateApplication.h"
+#include "Widgets/Layout/SBorder.h"
+#include "Widgets/Layout/SSpacer.h"
+#include "Widgets/SCanvas.h"
+#include "EditorStyleSet.h"
+#include "Editor/UnrealEdEngine.h"
+#include "UnrealEdGlobals.h"
+#include "LevelEditor.h"
+#include "Widgets/Docking/SDockTab.h"
+#include "ViewportTabContent.h"
+#if 0
+
+static const FName LevelEditorModName("LevelEditor");
+
+namespace ViewportLayoutDefs
+{
+	/** How many seconds to interpolate from restored to maximized state */
+	static const float MaximizeTransitionTime = 0.15f;
+
+	/** How many seconds to interpolate from maximized to restored state */
+	static const float RestoreTransitionTime = 0.2f;
+
+	/** Default maximized state for new layouts - will only be applied when no config data is restoring state */
+	static const bool bDefaultShouldBeMaximized = true;
+
+	/** Default immersive state for new layouts - will only be applied when no config data is restoring state */
+	static const bool bDefaultShouldBeImmersive = false;
+}
+
+// SViewportsOverlay ////////////////////////////////////////////////
+
+/**
+ * Overlay wrapper class so that we can cache the size of the widget
+ * It will also store the LevelViewportLayout data because that data can't be stored
+ * per app; it must be stored per viewport overlay in case the app that made it closes.
+ */
+class SViewportsOverlay : public SCompoundWidget
+{
+
+public:
+
+	SLATE_BEGIN_ARGS( SViewportsOverlay ){}
+		SLATE_DEFAULT_SLOT( FArguments, Content )
+		SLATE_ARGUMENT( TSharedPtr<FViewportTabContent>, LevelViewportTab )
+	SLATE_END_ARGS()
+
+	void Construct( const FArguments& InArgs );
+
+	/** Default constructor */
+	SViewportsOverlay()
+		: CachedSize( FVector2D::ZeroVector )
+	{}
+
+	/** Overridden from SWidget */
+	virtual void Tick( const FGeometry& AllottedGeometry, const double InCurrentTime, const float InDeltaTime ) override;
+
+	/** Wraps SOverlay::AddSlot() */
+	SOverlay::FOverlaySlot& AddSlot();
+
+	/** Wraps SOverlay::RemoveSlot() */
+	void RemoveSlot();
+
+	/**
+	 * Returns the cached size of this viewport overlay
+	 *
+	 * @return	The size that was cached
+	 */
+	const FVector2D& GetCachedSize() const;
+
+	/** Gets the Level Viewport Tab that created this overlay */
+	TSharedPtr<FViewportTabContent> GetLevelViewportTab() const;
+
+private:
+	
+	/** Reference to the owning level viewport tab */
+	TSharedPtr<FViewportTabContent> LevelViewportTab;
+
+	/** The overlay widget we're containing */
+	TSharedPtr< SOverlay > OverlayWidget;
+
+	/** Cache our size, so that we can use this when animating a viewport maximize/restore */
+	FVector2D CachedSize;
+};
+
+
+void SViewportsOverlay::Construct( const FArguments& InArgs )
+{
+	const TSharedRef<SWidget>& ContentWidget = InArgs._Content.Widget;
+	LevelViewportTab = InArgs._LevelViewportTab;
+
+	ChildSlot
+		[
+			SAssignNew( OverlayWidget, SOverlay )
+			+ SOverlay::Slot()
+			[
+				ContentWidget
+			]
+		];
+}
+
+void SViewportsOverlay::Tick( const FGeometry& AllottedGeometry, const double InCurrentTime, const float InDeltaTime )
+{
+	CachedSize = AllottedGeometry.Size;
+}
+
+SOverlay::FOverlaySlot& SViewportsOverlay::AddSlot()
+{
+	return OverlayWidget->AddSlot();
+}
+
+void SViewportsOverlay::RemoveSlot()
+{
+	return OverlayWidget->RemoveSlot();
+}
+
+const FVector2D& SViewportsOverlay::GetCachedSize() const
+{
+	return CachedSize;
+}
+
+TSharedPtr<FViewportTabContent> SViewportsOverlay::GetLevelViewportTab() const
+{
+	return LevelViewportTab;
+}
+
+#endif
\ No newline at end of file
diff --git a/Engine/Source/Editor/EditorFramework/Private/ViewportTabContent.cpp b/Engine/Source/Editor/EditorFramework/Private/ViewportTabContent.cpp
new file mode 100644
--- /dev/null
+++ b/Engine/Source/Editor/EditorFramework/Private/ViewportTabContent.cpp
@@ -0,0 +1,57 @@
+// Copyright Epic Games, Inc. All Rights Reserved.
+
+#include "ViewportTabContent.h"
+#include "Misc/ConfigCacheIni.h"
+#include "Framework/Docking/LayoutService.h"
+#include "Modules/ModuleManager.h"
+#include "Widgets/Docking/SDockTab.h"
+
+
+// FViewportTabContent ///////////////////////////
+
+bool FViewportTabContent::IsVisible() const
+{
+	if (ActiveViewportLayout.IsValid())
+	{
+		return ActiveViewportLayout->IsVisible();
+	}
+	return false;
+}
+
+bool FViewportTabContent::BelongsToTab(TSharedRef<class SDockTab> InParentTab) const
+{
+	TSharedPtr<SDockTab> ParentTabPinned = ParentTab.Pin();
+	return ParentTabPinned == InParentTab;
+}
+
+bool FViewportTabContent::IsViewportConfigurationSet(const FName& ConfigurationName) const
+{
+	if (ActiveViewportLayout.IsValid())
+	{
+		return ActiveViewportLayout->GetActivePaneConfigurationTypeName() == ConfigurationName;
+	}
+	return false;
+}
+
+const TMap< FName, TSharedPtr< IEditorViewportLayoutEntity > >* FViewportTabContent::GetViewports() const
+{
+	if (ActiveViewportLayout.IsValid())
+	{
+		return &ActiveViewportLayout->GetViewports();
+	}
+	return nullptr;
+}
+
+void FViewportTabContent::PerformActionOnViewports(ViewportActionFunction& TFuncPtr)
+{
+	const TMap< FName, TSharedPtr<IEditorViewportLayoutEntity> >* Entities = GetViewports();
+	if (!Entities)
+	{
+		return;
+	}
+
+	for (auto& Entity : *Entities)
+	{
+		TFuncPtr(Entity.Key, Entity.Value);
+	}
+}
diff --git a/Engine/Source/Editor/EditorFramework/Public/EditorViewportLayout.h b/Engine/Source/Editor/EditorFramework/Public/EditorViewportLayout.h
new file mode 100644
--- /dev/null
+++ b/Engine/Source/Editor/EditorFramework/Public/EditorViewportLayout.h
@@ -0,0 +1,56 @@
+// Copyright Epic Games, Inc. All Rights Reserved.
+
+#pragma once
+
+#include "CoreMinimal.h"
+#include "Widgets/SWidget.h"
+
+
+class IEditorViewportLayoutEntity : public TSharedFromThis<IEditorViewportLayoutEntity>
+{
+public:
+	/** Virtual destruction */
+	virtual ~IEditorViewportLayoutEntity() {};
+
+	/** Return a widget that is represents this entity */
+	virtual TSharedRef<SWidget> AsWidget() const = 0;
+
+	/** Set keyboard focus to this viewport entity */
+	virtual void SetKeyboardFocus() = 0;
+
+	/** Called when the parent layout is being destroyed */
+	virtual void OnLayoutDestroyed() = 0;
+
+	/** Called to save this item's settings in the specified config section */
+	virtual void SaveConfig(const FString& ConfigSection) = 0;
+
+	/** Get the type of this viewport as a name */
+	virtual FName GetType() const = 0;
+
+	/** Take a high res screen shot of viewport entity */
+	virtual void TakeHighResScreenShot() const = 0;
+};
+
+class FEditorViewportLayout 
+{
+public:
+	/** Virtual destruction */
+	virtual ~FEditorViewportLayout() = default;
+
+	virtual bool IsVisible() const { return true; }
+
+	/**
+	* @return All the viewports in this configuration
+	*/
+	const TMap< FName, TSharedPtr<IEditorViewportLayoutEntity> >& GetViewports() const { return Viewports; }
+
+	virtual void FactoryPaneConfigurationFromTypeName(const FName& InLayoutConfigTypeName) = 0;
+	virtual const FName GetActivePaneConfigurationTypeName() const = 0;
+
+	virtual void SaveConfig(const FString& LayoutSring) const {}
+	virtual void LoadConfig(const FString& LayoutString) {}
+
+protected:
+	/** List of all of the viewports in this layout, keyed on their config key */
+	TMap< FName, TSharedPtr< IEditorViewportLayoutEntity > > Viewports;
+};
diff --git a/Engine/Source/Editor/EditorFramework/Public/ViewportTabContent.h b/Engine/Source/Editor/EditorFramework/Public/ViewportTabContent.h
new file mode 100644
--- /dev/null
+++ b/Engine/Source/Editor/EditorFramework/Public/ViewportTabContent.h
@@ -0,0 +1,62 @@
+// Copyright Epic Games, Inc. All Rights Reserved.
+
+#pragma once
+
+#include "EditorViewportLayout.h"
+#include "Delegates/DelegateCombinations.h"
+#include "Templates/SharedPointer.h"
+
+/**
+ * Represents the content in a viewport tab in an editor.
+ * Each SDockTab holding viewports in an editor contains and owns one of these.
+ */
+
+class SDockTab;
+class FEditorViewportLayout;
+class IEditorViewportLayoutEntity;
+
+class EDITORFRAMEWORK_API FViewportTabContent
+{
+public:
+	virtual ~FViewportTabContent() {}
+
+	/** Returns whether the tab is currently shown */
+	bool IsVisible() const;
+
+	/** @return True if this viewport belongs to the tab given */
+	bool BelongsToTab(TSharedRef<SDockTab> InParentTab) const;
+
+	/**
+	* Returns whether the named layout is currently selected
+	*
+	* @param ConfigurationName		The name of the layout (for the names in namespace LevelViewportConfigurationNames)
+	* @return						True, if the named layout is currently active
+	*/
+	bool IsViewportConfigurationSet(const FName& ConfigurationName) const;
+
+	virtual void SetViewportConfiguration(const FName& ConfigurationName) {}
+
+	const TMap< FName, TSharedPtr< IEditorViewportLayoutEntity > >* GetViewports() const;
+
+	using ViewportActionFunction = TFunction<void(FName Name, TSharedPtr<IEditorViewportLayoutEntity>)>;
+	void PerformActionOnViewports(ViewportActionFunction& TFuncPtr);
+
+	DECLARE_EVENT(FViewportTabContent, FViewportTabContentLayoutChangedEvent);
+	virtual FViewportTabContentLayoutChangedEvent& OnViewportTabContentLayoutChanged() { return OnViewportTabContentLayoutChangedEvent; };
+
+	DECLARE_EVENT_OneParam(FViewportTabContent, FViewportTabContentLayoutStartChangeEvent, bool);
+	virtual FViewportTabContentLayoutStartChangeEvent& OnViewportTabContentLayoutStartChange() { return OnViewportTabContentLayoutStartChangeEvent; };
+
+protected:
+	FViewportTabContentLayoutChangedEvent OnViewportTabContentLayoutChangedEvent;
+	FViewportTabContentLayoutStartChangeEvent OnViewportTabContentLayoutStartChangeEvent;
+
+	TWeakPtr<SDockTab> ParentTab;
+
+	FString LayoutString;
+
+	/** Current layout */
+	TSharedPtr<FEditorViewportLayout> ActiveViewportLayout;
+
+	TOptional<FName> PreviouslyFocusedViewport;
+};
diff --git a/Engine/Source/Editor/UnrealEd/Private/EditorViewportLayout.cpp b/Engine/Source/Editor/UnrealEd/Private/EditorViewportLayout.cpp
deleted file mode 100644
diff --git a/Engine/Source/Editor/UnrealEd/Private/ViewportTabContent.cpp b/Engine/Source/Editor/UnrealEd/Private/ViewportTabContent.cpp
deleted file mode 100644
diff --git a/Engine/Source/Editor/UnrealEd/Public/EditorViewportLayout.h b/Engine/Source/Editor/UnrealEd/Public/EditorViewportLayout.h
deleted file mode 100644
diff --git a/Engine/Source/Editor/UnrealEd/Public/ViewportTabContent.h b/Engine/Source/Editor/UnrealEd/Public/ViewportTabContent.h
deleted file mode 100644
