Remove EditorSubsystem from EditorFramework build dependencies.

#rnx

diff --git a/Engine/Source/Editor/EditorFramework/EditorFramework.Build.cs b/Engine/Source/Editor/EditorFramework/EditorFramework.Build.cs
--- a/Engine/Source/Editor/EditorFramework/EditorFramework.Build.cs
+++ b/Engine/Source/Editor/EditorFramework/EditorFramework.Build.cs
@@ -21,7 +21,6 @@ public class EditorFramework : ModuleRules
 				"Core",
 				"CoreUObject",
 				"EditorStyle",
-				"EditorSubsystem",
 				"Engine",
 				"InputCore",
 				"RenderCore",
diff --git a/Engine/Source/Runtime/Engine/Engine.Build.cs b/Engine/Source/Runtime/Engine/Engine.Build.cs
--- a/Engine/Source/Runtime/Engine/Engine.Build.cs
+++ b/Engine/Source/Runtime/Engine/Engine.Build.cs
@@ -329,7 +329,6 @@ public class Engine : ModuleRules
 		{
 			PublicDependencyModuleNames.AddRange(
 				new string[] {
-					"EditorFramework",
 					"UnrealEd",
 					"Kismet"
 				}
