Fix for nounity errors

diff --git a/Engine/Source/Editor/EditorFramework/Private/Tools/Modes.cpp b/Engine/Source/Editor/EditorFramework/Private/Tools/Modes.cpp
--- a/Engine/Source/Editor/EditorFramework/Private/Tools/Modes.cpp
+++ b/Engine/Source/Editor/EditorFramework/Private/Tools/Modes.cpp
@@ -1,6 +1,7 @@
 // Copyright Epic Games, Inc. All Rights Reserved.
 
 #include "Tools/Modes.h"
+#include "EditorStyleSet.h"
 
 FEditorModeInfo::FEditorModeInfo()
 	: ID(NAME_None)
