Migrated undo/redo support from USelection to UTypedElementSelectionSet

This adds some new functions to UTypedElementSelectionInterface to deal with writing the underlying data needed by an element (WriteTransactedElement) to factory it again later (ReadTransactedElement) when restoring a prior selection state.

diff --git a/Engine/Source/Editor/EditorFramework/Private/Elements/TypedElementSelectionSet.cpp b/Engine/Source/Editor/EditorFramework/Private/Elements/TypedElementSelectionSet.cpp
--- a/Engine/Source/Editor/EditorFramework/Private/Elements/TypedElementSelectionSet.cpp
+++ b/Engine/Source/Editor/EditorFramework/Private/Elements/TypedElementSelectionSet.cpp
@@ -13,6 +13,88 @@ UTypedElementSelectionSet::UTypedElementSelectionSet()
 	}
 }
 
+bool UTypedElementSelectionSet::Modify(bool bAlwaysMarkDirty)
+{
+	if (GUndo && CanModify())
+	{
+		for (const TTypedElement<UTypedElementSelectionInterface>& SelectionElement : ElementList->IterateInterface<UTypedElementSelectionInterface>())
+		{
+			if (SelectionElement.ShouldPreventTransactions())
+			{
+				return false;
+			}
+		}
+
+		return Super::Modify(bAlwaysMarkDirty);
+	}
+
+	return false;
+}
+
+void UTypedElementSelectionSet::Serialize(FArchive& Ar)
+{
+	checkf(Ar.IsTransacting(), TEXT("UTypedElementSelectionSet can only be serialized for transactions!"));
+
+	if (Ar.IsSaving())
+	{
+		FTypedHandleTypeId ElementTypeId = 0;
+
+		for (const TTypedElement<UTypedElementSelectionInterface>& SelectionElement : ElementList->IterateInterface<UTypedElementSelectionInterface>())
+		{
+			ElementTypeId = SelectionElement.GetId().GetTypeId();
+			Ar << ElementTypeId;
+
+			SelectionElement.WriteTransactedElement(Ar);
+		}
+
+		// End of the list
+		ElementTypeId = 0;
+		Ar << ElementTypeId;
+	}
+	else if (Ar.IsLoading())
+	{
+		TArray<FTypedElementHandle, TInlineAllocator<256>> SelectedElements;
+
+		{
+			UTypedElementRegistry* Registry = UTypedElementRegistry::GetInstance();
+
+			FTypedHandleTypeId ElementTypeId = 0;
+			for (;;)
+			{
+				Ar << ElementTypeId;
+				if (!ElementTypeId)
+				{
+					// End of the list
+					break;
+				}
+
+				UTypedElementSelectionInterface* ElementTypeSelectionInterface = Registry->GetElementInterface<UTypedElementSelectionInterface>(ElementTypeId);
+				checkf(ElementTypeSelectionInterface, TEXT("Failed to find selection interface for a previously transacted element type!"));
+				FTypedElementHandle SelectedElement = ElementTypeSelectionInterface->ReadTransactedElement(Ar);
+				if (SelectedElement)
+				{
+					SelectedElements.Add(MoveTemp(SelectedElement));
+				}
+			}
+		}
+
+		{
+			FTypedElementListLegacySyncScopedBatch LegacySyncBatch(ElementList, /*bNotify*/false);
+			TGuardValue<bool> GuardIsRestoringFromTransaction(bIsRestoringFromTransaction, true);
+
+			// TODO: Work out the intersection of the before and after state instead of clearing and reselecting?
+
+			const FTypedElementSelectionOptions SelectionOptions = FTypedElementSelectionOptions()
+				.SetAllowHidden(true)
+				.SetAllowGroups(false)
+				.SetWarnIfLocked(false);
+
+			ClearSelection(SelectionOptions);
+			SelectElements(SelectedElements, SelectionOptions);
+		}
+	}
+}
+
 bool UTypedElementSelectionSet::IsElementSelected(const FTypedElementHandle& InElementHandle, const FTypedElementIsSelectedOptions InSelectionOptions) const
 {
 	FTypedElementSelectionSetElement SelectionSetElement = ResolveSelectionSetElement(InElementHandle);
@@ -129,7 +211,7 @@ FTypedElementHandle UTypedElementSelectionSet::GetSelectionElement(const FTypedE
 
 void UTypedElementSelectionSet::RegisterAssetEditorSelectionProxy(const FName InElementTypeName, UTypedElementAssetEditorSelectionProxy* InAssetEditorSelectionProxy)
 {
-	const FTypedHandleTypeId ElementTypeId = ElementList->GetRegisteredElementTypeId(InElementTypeName);
+	const FTypedHandleTypeId ElementTypeId = UTypedElementRegistry::GetInstance()->GetRegisteredElementTypeId(InElementTypeName);
 	checkf(ElementTypeId > 0, TEXT("Element type '%s' has not been registered!"), *InElementTypeName.ToString());
 	RegisteredAssetEditorSelectionProxies[ElementTypeId - 1] = InAssetEditorSelectionProxy;
 }
@@ -143,3 +225,21 @@ FTypedElementSelectionSetElement UTypedElementSelectionSet::ResolveSelectionSetE
 	}
 	return FTypedElementSelectionSetElement();
 }
+
+void UTypedElementSelectionSet::OnElementListPreChange(const UTypedElementList* InElementList)
+{
+	check(InElementList == ElementList);
+	OnPreChangeDelegate.Broadcast(this);
+
+	if (!bIsRestoringFromTransaction)
+	{
+		// Track the pre-change state for undo/redo
+		Modify();
+	}
+}
+
+void UTypedElementSelectionSet::OnElementListChanged(const UTypedElementList* InElementList)
+{
+	check(InElementList == ElementList);
+	OnChangedDelegate.Broadcast(this);
+}
diff --git a/Engine/Source/Editor/EditorFramework/Public/Elements/TypedElementSelectionSet.h b/Engine/Source/Editor/EditorFramework/Public/Elements/TypedElementSelectionSet.h
--- a/Engine/Source/Editor/EditorFramework/Public/Elements/TypedElementSelectionSet.h
+++ b/Engine/Source/Editor/EditorFramework/Public/Elements/TypedElementSelectionSet.h
@@ -86,6 +86,10 @@ class EDITORFRAMEWORK_API UTypedElementSelectionSet : public UObject
 public:
 	UTypedElementSelectionSet();
 
+	//~ UObject interface
+	virtual bool Modify(bool bAlwaysMarkDirty = true) override;
+	virtual void Serialize(FArchive& Ar) override;
+
 	/**
 	 * Test to see whether the given element is currently considered selected.
 	 */
@@ -229,20 +233,12 @@ private:
 	/**
 	 * Proxy the internal OnPreChange event from the underlying element list.
 	 */
-	void OnElementListPreChange(const UTypedElementList* InElementList) const
-	{
-		check(InElementList == ElementList);
-		OnPreChangeDelegate.Broadcast(this);
-	}
+	void OnElementListPreChange(const UTypedElementList* InElementList);
 
 	/**
 	 * Proxy the internal OnChanged event from the underlying element list.
 	 */
-	void OnElementListChanged(const UTypedElementList* InElementList) const
-	{
-		check(InElementList == ElementList);
-		OnChangedDelegate.Broadcast(this);
-	}
+	void OnElementListChanged(const UTypedElementList* InElementList);
 
 	/** Underlying element list holding the selection state. */
 	UPROPERTY()
@@ -257,4 +253,7 @@ private:
 
 	/** Delegate that is invoked whenever the underlying element list has been changed. */
 	FOnChanged OnChangedDelegate;
+
+	/** Set when we are currently restoring the selection state from a transaction (undo/redo) */
+	bool bIsRestoringFromTransaction = false;
 };
diff --git a/Engine/Source/Editor/LevelEditor/Private/SLevelEditor.cpp b/Engine/Source/Editor/LevelEditor/Private/SLevelEditor.cpp
--- a/Engine/Source/Editor/LevelEditor/Private/SLevelEditor.cpp
+++ b/Engine/Source/Editor/LevelEditor/Private/SLevelEditor.cpp
@@ -217,7 +217,7 @@ void SLevelEditor::Construct( const SLevelEditor::FArguments& InArgs)
 
 void SLevelEditor::Initialize( const TSharedRef<SDockTab>& OwnerTab, const TSharedRef<SWindow>& OwnerWindow )
 {
-	SelectedElements = NewObject<UTypedElementSelectionSet>();
+	SelectedElements = NewObject<UTypedElementSelectionSet>(GetTransientPackage(), NAME_None, RF_Transactional);
 	SelectedElements->AddToRoot();
 
 	// Register the level editor specific selection behavior
diff --git a/Engine/Source/Editor/UnrealEd/Private/AssetEditorModeManager.cpp b/Engine/Source/Editor/UnrealEd/Private/AssetEditorModeManager.cpp
--- a/Engine/Source/Editor/UnrealEd/Private/AssetEditorModeManager.cpp
+++ b/Engine/Source/Editor/UnrealEd/Private/AssetEditorModeManager.cpp
@@ -11,7 +11,7 @@
 FAssetEditorModeManager::FAssetEditorModeManager()
 {
 	{
-		UTypedElementSelectionSet* ActorAndComponentsSelectionSet = NewObject<UTypedElementSelectionSet>();
+		UTypedElementSelectionSet* ActorAndComponentsSelectionSet = NewObject<UTypedElementSelectionSet>(GetTransientPackage(), NAME_None, RF_Transactional);
 
 		ActorSet = USelection::CreateActorSelection(GetTransientPackage(), NAME_None, RF_Transactional);
 		ActorSet->SetElementSelectionSet(ActorAndComponentsSelectionSet);
@@ -23,7 +23,7 @@ FAssetEditorModeManager::FAssetEditorModeManager()
 	}
 
 	ObjectSet = USelection::CreateObjectSelection(GetTransientPackage(), NAME_None, RF_Transactional);
-	ObjectSet->SetElementSelectionSet(NewObject<UTypedElementSelectionSet>(ObjectSet));
+	ObjectSet->SetElementSelectionSet(NewObject<UTypedElementSelectionSet>(ObjectSet, NAME_None, RF_Transactional));
 	ObjectSet->AddToRoot();
 }
 
diff --git a/Engine/Source/Editor/UnrealEd/Private/EditorEngine.cpp b/Engine/Source/Editor/UnrealEd/Private/EditorEngine.cpp
--- a/Engine/Source/Editor/UnrealEd/Private/EditorEngine.cpp
+++ b/Engine/Source/Editor/UnrealEd/Private/EditorEngine.cpp
@@ -268,7 +268,7 @@ void InitSelectionSets()
 	GObjectSelection = USelection::CreateObjectSelection(GetTransientPackage(), TEXT("SelectedObjects"), RF_Transactional);
 	GObjectSelection->AddToRoot();
 
-	UTypedElementSelectionSet* ObjectSelectionSet = NewObject<UTypedElementSelectionSet>(GObjectSelection);
+	UTypedElementSelectionSet* ObjectSelectionSet = NewObject<UTypedElementSelectionSet>(GObjectSelection, NAME_None, RF_Transactional);
 	ObjectSelectionSet->RegisterAssetEditorSelectionProxy(NAME_Object, NewObject<UObjectElementEditorSelectionProxy>());
 	GObjectSelection->SetElementSelectionSet(ObjectSelectionSet);
 }
diff --git a/Engine/Source/Editor/UnrealEd/Private/Elements/Actor/ActorElementEditorSelectionInterface.cpp b/Engine/Source/Editor/UnrealEd/Private/Elements/Actor/ActorElementEditorSelectionInterface.cpp
--- a/Engine/Source/Editor/UnrealEd/Private/Elements/Actor/ActorElementEditorSelectionInterface.cpp
+++ b/Engine/Source/Editor/UnrealEd/Private/Elements/Actor/ActorElementEditorSelectionInterface.cpp
@@ -3,8 +3,10 @@
 #include "Elements/Actor/ActorElementEditorSelectionInterface.h"
 #include "Elements/Actor/ActorElementData.h"
 #include "GameFramework/Actor.h"
+
 #include "TypedElementList.h"
 #include "Elements/EngineElementsLibrary.h"
+#include "Elements/Object/ObjectElementEditorSelectionInterface.h"
 
 bool UActorElementEditorSelectionInterface::IsElementSelected(const FTypedElementHandle& InElementHandle, const UTypedElementList* InSelectionSet, const FTypedElementIsSelectedOptions& InSelectionOptions)
 {
@@ -12,6 +14,29 @@ bool UActorElementEditorSelectionInterface::IsElementSelected(const FTypedElemen
 	return ActorData && IsActorSelected(ActorData->Actor, InSelectionSet, InSelectionOptions);
 }
 
+bool UActorElementEditorSelectionInterface::ShouldPreventTransactions(const FTypedElementHandle& InElementHandle)
+{
+	const FActorElementData* ActorData = InElementHandle.GetData<FActorElementData>();
+	return ActorData && UObjectElementEditorSelectionInterface::ShouldObjectPreventTransactions(ActorData->Actor);
+}
+
+void UActorElementEditorSelectionInterface::WriteTransactedElement(const FTypedElementHandle& InElementHandle, FArchive& InArchive)
+{
+	const FActorElementData* ActorData = InElementHandle.GetData<FActorElementData>();
+	if (ActorData)
+	{
+		UObjectElementEditorSelectionInterface::WriteTransactedObject(ActorData->Actor, InArchive);
+	}
+}
+
+FTypedElementHandle UActorElementEditorSelectionInterface::ReadTransactedElement(FArchive& InArchive)
+{
+	return UObjectElementEditorSelectionInterface::ReadTransactedObject(InArchive, [](const UObject* InObject)
+	{
+		return UEngineElementsLibrary::AcquireEditorActorElementHandle(CastChecked<const AActor>(InObject));
+	});
+}
+
 bool UActorElementEditorSelectionInterface::IsActorSelected(const AActor* InActor, const UTypedElementList* InSelectionSet, const FTypedElementIsSelectedOptions& InSelectionOptions)
 {
 	if (InSelectionSet->Num() == 0)
diff --git a/Engine/Source/Editor/UnrealEd/Private/Elements/Actor/ActorElementEditorSelectionInterface.h b/Engine/Source/Editor/UnrealEd/Private/Elements/Actor/ActorElementEditorSelectionInterface.h
--- a/Engine/Source/Editor/UnrealEd/Private/Elements/Actor/ActorElementEditorSelectionInterface.h
+++ b/Engine/Source/Editor/UnrealEd/Private/Elements/Actor/ActorElementEditorSelectionInterface.h
@@ -14,6 +14,9 @@ class UActorElementEditorSelectionInterface : public UActorElementSelectionInter
 
 public:
 	virtual bool IsElementSelected(const FTypedElementHandle& InElementHandle, const UTypedElementList* InSelectionSet, const FTypedElementIsSelectedOptions& InSelectionOptions) override;
+	virtual bool ShouldPreventTransactions(const FTypedElementHandle& InElementHandle) override;
+	virtual void WriteTransactedElement(const FTypedElementHandle& InElementHandle, FArchive& InArchive) override;
+	virtual FTypedElementHandle ReadTransactedElement(FArchive& InArchive) override;
 
 	static bool IsActorSelected(const AActor* InActor, const UTypedElementList* InSelectionSet, const FTypedElementIsSelectedOptions& InSelectionOptions);
 };
diff --git a/Engine/Source/Editor/UnrealEd/Private/Elements/Component/ComponentElementEditorSelectionInterface.cpp b/Engine/Source/Editor/UnrealEd/Private/Elements/Component/ComponentElementEditorSelectionInterface.cpp
--- a/Engine/Source/Editor/UnrealEd/Private/Elements/Component/ComponentElementEditorSelectionInterface.cpp
+++ b/Engine/Source/Editor/UnrealEd/Private/Elements/Component/ComponentElementEditorSelectionInterface.cpp
@@ -4,8 +4,10 @@
 #include "Elements/Component/ComponentElementData.h"
 #include "Components/SceneComponent.h"
 #include "GameFramework/Actor.h"
+
 #include "TypedElementList.h"
 #include "Elements/EngineElementsLibrary.h"
+#include "Elements/Object/ObjectElementEditorSelectionInterface.h"
 
 bool UComponentElementEditorSelectionInterface::IsElementSelected(const FTypedElementHandle& InElementHandle, const UTypedElementList* InSelectionSet, const FTypedElementIsSelectedOptions& InSelectionOptions)
 {
@@ -13,6 +15,29 @@ bool UComponentElementEditorSelectionInterface::IsElementSelected(const FTypedEl
 	return ComponentData && IsComponentSelected(ComponentData->Component, InSelectionSet, InSelectionOptions);
 }
 
+bool UComponentElementEditorSelectionInterface::ShouldPreventTransactions(const FTypedElementHandle& InElementHandle)
+{
+	const FComponentElementData* ComponentData = InElementHandle.GetData<FComponentElementData>();
+	return ComponentData && UObjectElementEditorSelectionInterface::ShouldObjectPreventTransactions(ComponentData->Component);
+}
+
+void UComponentElementEditorSelectionInterface::WriteTransactedElement(const FTypedElementHandle& InElementHandle, FArchive& InArchive)
+{
+	const FComponentElementData* ComponentData = InElementHandle.GetData<FComponentElementData>();
+	if (ComponentData)
+	{
+		UObjectElementEditorSelectionInterface::WriteTransactedObject(ComponentData->Component, InArchive);
+	}
+}
+
+FTypedElementHandle UComponentElementEditorSelectionInterface::ReadTransactedElement(FArchive& InArchive)
+{
+	return UObjectElementEditorSelectionInterface::ReadTransactedObject(InArchive, [](const UObject* InObject)
+	{
+		return UEngineElementsLibrary::AcquireEditorComponentElementHandle(CastChecked<const UActorComponent>(InObject));
+	});
+}
+
 bool UComponentElementEditorSelectionInterface::IsComponentSelected(const UActorComponent* InComponent, const UTypedElementList* InSelectionSet, const FTypedElementIsSelectedOptions& InSelectionOptions)
 {
 	if (InSelectionSet->Num() == 0)
diff --git a/Engine/Source/Editor/UnrealEd/Private/Elements/Component/ComponentElementEditorSelectionInterface.h b/Engine/Source/Editor/UnrealEd/Private/Elements/Component/ComponentElementEditorSelectionInterface.h
--- a/Engine/Source/Editor/UnrealEd/Private/Elements/Component/ComponentElementEditorSelectionInterface.h
+++ b/Engine/Source/Editor/UnrealEd/Private/Elements/Component/ComponentElementEditorSelectionInterface.h
@@ -14,6 +14,9 @@ class UComponentElementEditorSelectionInterface : public UComponentElementSelect
 
 public:
 	virtual bool IsElementSelected(const FTypedElementHandle& InElementHandle, const UTypedElementList* InSelectionSet, const FTypedElementIsSelectedOptions& InSelectionOptions) override;
+	virtual bool ShouldPreventTransactions(const FTypedElementHandle& InElementHandle) override;
+	virtual void WriteTransactedElement(const FTypedElementHandle& InElementHandle, FArchive& InArchive) override;
+	virtual FTypedElementHandle ReadTransactedElement(FArchive& InArchive) override;
 
 	static bool IsComponentSelected(const UActorComponent* InComponent, const UTypedElementList* InSelectionSet, const FTypedElementIsSelectedOptions& InSelectionOptions);
 };
diff --git a/Engine/Source/Editor/UnrealEd/Private/Elements/EditorElements.cpp b/Engine/Source/Editor/UnrealEd/Private/Elements/EditorElements.cpp
--- a/Engine/Source/Editor/UnrealEd/Private/Elements/EditorElements.cpp
+++ b/Engine/Source/Editor/UnrealEd/Private/Elements/EditorElements.cpp
@@ -3,12 +3,21 @@
 #include "Elements/EditorElements.h"
 #include "TypedElementRegistry.h"
 
+#include "Elements/Object/ObjectElementEditorSelectionInterface.h"
+
 #include "Elements/Actor/ActorElementEditorSelectionInterface.h"
 
 #include "Elements/Component/ComponentElementEditorSelectionInterface.h"
 
 FSimpleMulticastDelegate OnRegisterEditorElementsDelegate;
 
+void RegisterEditorObjectElements()
+{
+	UTypedElementRegistry* Registry = UTypedElementRegistry::GetInstance();
+
+	Registry->RegisterElementInterface<UTypedElementSelectionInterface>(NAME_Object, NewObject<UObjectElementEditorSelectionInterface>(), /*bAllowOverride*/true);
+}
+
 void RegisterEditorActorElements()
 {
 	UTypedElementRegistry* Registry = UTypedElementRegistry::GetInstance();
@@ -25,6 +34,7 @@ void RegisterEditorComponentElements()
 
 void RegisterEditorElements()
 {
+	RegisterEditorObjectElements();
 	RegisterEditorActorElements();
 	RegisterEditorComponentElements();
 
diff --git a/Engine/Source/Editor/UnrealEd/Private/Elements/Object/ObjectElementEditorSelectionInterface.cpp b/Engine/Source/Editor/UnrealEd/Private/Elements/Object/ObjectElementEditorSelectionInterface.cpp
new file mode 100644
--- /dev/null
+++ b/Engine/Source/Editor/UnrealEd/Private/Elements/Object/ObjectElementEditorSelectionInterface.cpp
@@ -0,0 +1,54 @@
+// Copyright Epic Games, Inc. All Rights Reserved.
+
+#include "Elements/Object/ObjectElementEditorSelectionInterface.h"
+#include "Elements/Object/ObjectElementData.h"
+
+#include "UObject/Package.h"
+#include "Elements/EngineElementsLibrary.h"
+
+bool UObjectElementEditorSelectionInterface::ShouldPreventTransactions(const FTypedElementHandle& InElementHandle)
+{
+	const FObjectElementData* ObjectData = InElementHandle.GetData<FObjectElementData>();
+	return ObjectData && ShouldObjectPreventTransactions(ObjectData->Object);
+}
+
+void UObjectElementEditorSelectionInterface::WriteTransactedElement(const FTypedElementHandle& InElementHandle, FArchive& InArchive)
+{
+	const FObjectElementData* ObjectData = InElementHandle.GetData<FObjectElementData>();
+	if (ObjectData)
+	{
+		WriteTransactedObject(ObjectData->Object, InArchive);
+	}
+}
+
+FTypedElementHandle UObjectElementEditorSelectionInterface::ReadTransactedElement(FArchive& InArchive)
+{
+	return ReadTransactedObject(InArchive, [](const UObject* InObject)
+	{
+		check(InObject);
+		return UEngineElementsLibrary::AcquireEditorObjectElementHandle(InObject);
+	});
+}
+
+bool UObjectElementEditorSelectionInterface::ShouldObjectPreventTransactions(const UObject* InObject)
+{
+	// If the selection currently contains any PIE objects we should not be including it in the transaction buffer
+	return InObject->GetOutermost()->HasAnyPackageFlags(PKG_PlayInEditor | PKG_ContainsScript | PKG_CompiledIn);
+}
+
+void UObjectElementEditorSelectionInterface::WriteTransactedObject(const UObject* InObject, FArchive& InArchive)
+{
+	TWeakObjectPtr<const UObject> SelectedObject = InObject;
+	InArchive << SelectedObject;
+}
+
+FTypedElementHandle UObjectElementEditorSelectionInterface::ReadTransactedObject(FArchive& InArchive, TFunctionRef<FTypedElementHandle(const UObject*)> InFindElementForObject)
+{
+	TWeakObjectPtr<const UObject> SelectedObject;
+	InArchive << SelectedObject;
+
+	const UObject* SelectedObjectPtr = SelectedObject.Get(/*bEvenIfPendingKill*/true);
+	return SelectedObjectPtr
+		? InFindElementForObject(SelectedObjectPtr)
+		: FTypedElementHandle();
+}
diff --git a/Engine/Source/Editor/UnrealEd/Private/Elements/Object/ObjectElementEditorSelectionInterface.h b/Engine/Source/Editor/UnrealEd/Private/Elements/Object/ObjectElementEditorSelectionInterface.h
new file mode 100644
--- /dev/null
+++ b/Engine/Source/Editor/UnrealEd/Private/Elements/Object/ObjectElementEditorSelectionInterface.h
@@ -0,0 +1,21 @@
+// Copyright Epic Games, Inc. All Rights Reserved.
+
+#pragma once
+
+#include "Elements/Object/ObjectElementSelectionInterface.h"
+#include "ObjectElementEditorSelectionInterface.generated.h"
+
+UCLASS()
+class UObjectElementEditorSelectionInterface : public UObjectElementSelectionInterface
+{
+	GENERATED_BODY()
+
+public:
+	virtual bool ShouldPreventTransactions(const FTypedElementHandle& InElementHandle) override;
+	virtual void WriteTransactedElement(const FTypedElementHandle& InElementHandle, FArchive& InArchive) override;
+	virtual FTypedElementHandle ReadTransactedElement(FArchive& InArchive) override;
+
+	static bool ShouldObjectPreventTransactions(const UObject* InObject);
+	static void WriteTransactedObject(const UObject* InObject, FArchive& InArchive);
+	static FTypedElementHandle ReadTransactedObject(FArchive& InArchive, TFunctionRef<FTypedElementHandle(const UObject*)> InFindElementForObject);
+};
diff --git a/Engine/Source/Editor/UnrealEd/Private/Selection.cpp b/Engine/Source/Editor/UnrealEd/Private/Selection.cpp
--- a/Engine/Source/Editor/UnrealEd/Private/Selection.cpp
+++ b/Engine/Source/Editor/UnrealEd/Private/Selection.cpp
@@ -100,7 +100,6 @@ void USelection::SetElementSelectionSet(UTypedElementSelectionSet* InElementSele
 {
 	if (ElementSelectionSet)
 	{
-		ElementSelectionSet->OnPreChange().RemoveAll(this);
 		ElementSelectionSet->Legacy_GetElementListSync().OnSyncEvent().RemoveAll(this);
 	}
 
@@ -109,7 +108,6 @@ void USelection::SetElementSelectionSet(UTypedElementSelectionSet* InElementSele
 
 	if (ElementSelectionSet)
 	{
-		ElementSelectionSet->OnPreChange().AddUObject(this, &USelection::OnElementPreChangeEvent);
 		ElementSelectionSet->Legacy_GetElementListSync().OnSyncEvent().AddUObject(this, &USelection::OnElementListSyncEvent);
 	}
 	
@@ -201,12 +199,6 @@ UObject* USelection::GetObjectForElementHandle(const FTypedElementHandle& InElem
 	return nullptr;
 }
 
-void USelection::OnElementPreChangeEvent(const UTypedElementSelectionSet* InElementSelectionSet)
-{
-	check(InElementSelectionSet == ElementSelectionSet);
-	Modify();
-}
-
 void USelection::OnElementListSyncEvent(const UTypedElementList* InElementList, FTypedElementListLegacySync::ESyncType InSyncType, const FTypedElementHandle& InElementHandle, bool bIsWithinBatchOperation)
 {
 	check(InElementList == ElementSelectionSet->GetElementList());
@@ -421,47 +413,14 @@ bool USelection::IsSelected(const UObject* InObject) const
 
 void USelection::Serialize(FArchive& Ar)
 {
-	Super::Serialize(Ar);
-
-	TArray<TWeakObjectPtr<UObject>> SelectedObjects;
-	if (Ar.IsSaving())
+	if (ElementSelectionSet)
 	{
-		GetSelectedObjects(SelectedObjects); // TODO: The old code used to use bEvenIfPendingKill during the resolve
-	}
-	Ar << SelectedObjects;
-
-	if (Ar.IsLoading())
-	{
-		BeginBatchSelectOperation();
-
-		// The set of selected objects may have changed, so make sure our annotations exactly match the list, otherwise
-		// UObject::IsSelected() could return a result that was different from the list of objects returned by GetSelectedObjects()
-		// This needs to happen in serialize because other code may check the selection state in PostEditUndo and the order of PostEditUndo is indeterminate.
-		DeselectAll(nullptr);
-		for (const TWeakObjectPtr<UObject>& ObjectPtr : SelectedObjects)
-		{
-			if (UObject* Object = ObjectPtr.Get(true))
-			{
-				Select(Object);
-			}
-		}
-		SyncSelectedClasses();
-
-		EndBatchSelectOperation(/*bNotify*/false);
+		ElementSelectionSet->Serialize(Ar);
 	}
 }
 
-bool USelection::Modify(bool bAlwaysMarkDirty/* =true */)
+bool USelection::Modify(bool bAlwaysMarkDirty)
 {
-	// If the selection currently contains any PIE objects we should not be including it in the transaction buffer
-	for (int32 Idx = 0; Idx < Num(); ++Idx)
-	{
-		UObject* Object = GetSelectedObject(Idx);
-		if (Object && Object->GetOutermost()->HasAnyPackageFlags(PKG_PlayInEditor | PKG_ContainsScript | PKG_CompiledIn))
-		{
-			return false;
-		}
-	}
-
-	return Super::Modify(bAlwaysMarkDirty);
+	return ElementSelectionSet 
+		&& ElementSelectionSet->Modify(bAlwaysMarkDirty);
 }
diff --git a/Engine/Source/Editor/UnrealEd/Public/Selection.h b/Engine/Source/Editor/UnrealEd/Public/Selection.h
--- a/Engine/Source/Editor/UnrealEd/Public/Selection.h
+++ b/Engine/Source/Editor/UnrealEd/Public/Selection.h
@@ -365,7 +365,6 @@ private:
 	bool IsValidObjectToSelect(const UObject* InObject) const;
 	UObject* GetObjectForElementHandle(const FTypedElementHandle& InElementHandle) const;
 
-	void OnElementPreChangeEvent(const UTypedElementSelectionSet* InElementSelectionSet);
 	void OnElementListSyncEvent(const UTypedElementList* InElementList, FTypedElementListLegacySync::ESyncType InSyncType, const FTypedElementHandle& InElementHandle, bool bIsWithinBatchOperation);
 
 	void OnObjectSelected(UObject* InObject, const bool bNotify);
diff --git a/Engine/Source/Runtime/TypedElementFramework/Private/TypedElementList.cpp b/Engine/Source/Runtime/TypedElementFramework/Private/TypedElementList.cpp
--- a/Engine/Source/Runtime/TypedElementFramework/Private/TypedElementList.cpp
+++ b/Engine/Source/Runtime/TypedElementFramework/Private/TypedElementList.cpp
@@ -131,11 +131,6 @@ UTypedElementInterface* UTypedElementList::GetElementInterface(const FTypedEleme
 	return Registry->GetElementInterface(InElementHandle, InBaseInterfaceType);
 }
 
-FTypedHandleTypeId UTypedElementList::GetRegisteredElementTypeId(const FName InElementTypeName) const
-{
-	return Registry->GetRegisteredElementTypeId(InElementTypeName);
-}
-
 bool UTypedElementList::AddElementImpl(FTypedElementHandle&& InElementHandle)
 {
 	if (!InElementHandle)
diff --git a/Engine/Source/Runtime/TypedElementFramework/Private/TypedElementRegistry.cpp b/Engine/Source/Runtime/TypedElementFramework/Private/TypedElementRegistry.cpp
--- a/Engine/Source/Runtime/TypedElementFramework/Private/TypedElementRegistry.cpp
+++ b/Engine/Source/Runtime/TypedElementFramework/Private/TypedElementRegistry.cpp
@@ -99,15 +99,15 @@ void UTypedElementRegistry::RegisterElementInterfaceImpl(const FName InElementTy
 	RegisteredElementType->Interfaces.Add(InBaseInterfaceType->GetFName(), InElementInterface);
 }
 
-UTypedElementInterface* UTypedElementRegistry::GetElementInterfaceImpl(const FTypedElementId& InElementId, const TSubclassOf<UTypedElementInterface>& InBaseInterfaceType) const
+UTypedElementInterface* UTypedElementRegistry::GetElementInterfaceImpl(const FTypedHandleTypeId InElementTypeId, const TSubclassOf<UTypedElementInterface>& InBaseInterfaceType) const
 {
-	if (!InElementId)
+	if (!InElementTypeId)
 	{
 		return nullptr;
 	}
 
-	FRegisteredElementType* RegisteredElementType = GetRegisteredElementTypeFromId(InElementId.GetTypeId());
-	checkf(RegisteredElementType, TEXT("Element type ID '%d' has not been registered!"), InElementId.GetTypeId());
+	FRegisteredElementType* RegisteredElementType = GetRegisteredElementTypeFromId(InElementTypeId);
+	checkf(RegisteredElementType, TEXT("Element type ID '%d' has not been registered!"), InElementTypeId);
 
 	return RegisteredElementType->Interfaces.FindRef(InBaseInterfaceType->GetFName());
 }
diff --git a/Engine/Source/Runtime/TypedElementFramework/Public/TypedElementList.h b/Engine/Source/Runtime/TypedElementFramework/Public/TypedElementList.h
--- a/Engine/Source/Runtime/TypedElementFramework/Public/TypedElementList.h
+++ b/Engine/Source/Runtime/TypedElementFramework/Public/TypedElementList.h
@@ -531,12 +531,6 @@ public:
 		});
 	}
 
-	/**
-	 * Get the element type ID for the associated element type name, if any.
-	 * @return The element type ID, or 0 if the given name wasn't registered.
-	 */
-	FTypedHandleTypeId GetRegisteredElementTypeId(const FName InElementTypeName) const;
-
 	/**
 	 * Access the delegate that is invoked whenever this element list is potentially about to change.
 	 * @note This may be called even if no actual change happens, though once a change does happen it won't be called again until after the next call to NotifyPendingChanges.
diff --git a/Engine/Source/Runtime/TypedElementFramework/Public/TypedElementRegistry.h b/Engine/Source/Runtime/TypedElementFramework/Public/TypedElementRegistry.h
--- a/Engine/Source/Runtime/TypedElementFramework/Public/TypedElementRegistry.h
+++ b/Engine/Source/Runtime/TypedElementFramework/Public/TypedElementRegistry.h
@@ -80,12 +80,12 @@ public:
 	}
 
 	/**
-	 * Get the element interface supported by the given handle, or null if there is no support for this interface.
+	 * Get the element interface supported by the given type, or null if there is no support for this interface.
 	 */
 	template <typename BaseInterfaceType>
-	FORCEINLINE BaseInterfaceType* GetElementInterface(const FTypedElementId& InElementId) const
+	FORCEINLINE BaseInterfaceType* GetElementInterface(const FTypedHandleTypeId InElementTypeId) const
 	{
-		return static_cast<BaseInterfaceType*>(GetElementInterfaceImpl(InElementId, BaseInterfaceType::StaticClass()));
+		return static_cast<BaseInterfaceType*>(GetElementInterfaceImpl(InElementTypeId, BaseInterfaceType::StaticClass()));
 	}
 
 	/**
@@ -94,7 +94,7 @@ public:
 	template <typename BaseInterfaceType>
 	FORCEINLINE BaseInterfaceType* GetElementInterface(const FTypedElementHandle& InElementHandle) const
 	{
-		return static_cast<BaseInterfaceType*>(GetElementInterfaceImpl(InElementHandle.GetId(), BaseInterfaceType::StaticClass()));
+		return static_cast<BaseInterfaceType*>(GetElementInterfaceImpl(InElementHandle.GetId().GetTypeId(), BaseInterfaceType::StaticClass()));
 	}
 
 	/**
@@ -103,7 +103,7 @@ public:
 	UFUNCTION(BlueprintPure, Category="TypedElementFramework|Registry")
 	FORCEINLINE UTypedElementInterface* GetElementInterface(const FTypedElementHandle& InElementHandle, const TSubclassOf<UTypedElementInterface>& InBaseInterfaceType) const
 	{
-		return GetElementInterfaceImpl(InElementHandle.GetId(), InBaseInterfaceType);
+		return GetElementInterfaceImpl(InElementHandle.GetId().GetTypeId(), InBaseInterfaceType);
 	}
 
 	/**
@@ -308,7 +308,7 @@ private:
 
 	void RegisterElementTypeImpl(const FName InElementTypeName, TUniquePtr<FRegisteredElementType>&& InRegisteredElementType);
 	void RegisterElementInterfaceImpl(const FName InElementTypeName, UTypedElementInterface* InElementInterface, const TSubclassOf<UTypedElementInterface>& InBaseInterfaceType, const bool InAllowOverride);
-	UTypedElementInterface* GetElementInterfaceImpl(const FTypedElementId& InElementId, const TSubclassOf<UTypedElementInterface>& InBaseInterfaceType) const;
+	UTypedElementInterface* GetElementInterfaceImpl(const FTypedHandleTypeId InElementTypeId, const TSubclassOf<UTypedElementInterface>& InBaseInterfaceType) const;
 
 	template <typename ElementDataType>
 	TTypedElementOwner<ElementDataType> CreateElementImpl(const FName InElementTypeName, const FTypedHandleElementId InElementId)
diff --git a/Engine/Source/Runtime/TypedElementInterfaces/Public/TypedElementSelectionInterface.h b/Engine/Source/Runtime/TypedElementInterfaces/Public/TypedElementSelectionInterface.h
--- a/Engine/Source/Runtime/TypedElementInterfaces/Public/TypedElementSelectionInterface.h
+++ b/Engine/Source/Runtime/TypedElementInterfaces/Public/TypedElementSelectionInterface.h
@@ -111,6 +111,23 @@ public:
 	 */
 	UFUNCTION(BlueprintPure, Category="TypedElementInterfaces|Selection")
 	virtual FTypedElementHandle GetSelectionElement(const FTypedElementHandle& InElementHandle, const UTypedElementList* InCurrentSelection, const ETypedElementSelectionMethod InSelectionMethod) { return InElementHandle; }
+
+	/**
+	 * Test to see whether the given element prevents the selection set state from being transacted for undo/redo (eg, if the element belongs to a PIE instance).
+	 */
+	virtual bool ShouldPreventTransactions(const FTypedElementHandle& InElementHandle) { return false; }
+
+	/**
+	 * Write the information needed to find the given element again when this transaction is replayed for undo/redo.
+	 * @note The data written must match the data read by ReadTransactedElement.
+	 */
+	virtual void WriteTransactedElement(const FTypedElementHandle& InElementHandle, FArchive& InArchive) {}
+
+	/**
+	 * Read the information needed to be find a previously tracked element, and attempt to get it now that this transaction is being replayed for undo/redo.
+	 * @note The data read must match the data written by WriteTransactedElement.
+	 */
+	virtual FTypedElementHandle ReadTransactedElement(FArchive& InArchive) { return FTypedElementHandle(); }
 };
 
 template <>
@@ -123,4 +140,7 @@ struct TTypedElement<UTypedElementSelectionInterface> : public TTypedElementBase
 	bool DeselectElement(UTypedElementList* InSelectionSet, const FTypedElementSelectionOptions& InSelectionOptions) const { return InterfacePtr->DeselectElement(*this, InSelectionSet, InSelectionOptions); }
 	bool AllowSelectionModifiers(const UTypedElementList* InSelectionSet) const { return InterfacePtr->AllowSelectionModifiers(*this, InSelectionSet); }
 	FTypedElementHandle GetSelectionElement(const UTypedElementList* InCurrentSelection, const ETypedElementSelectionMethod InSelectionMethod) const { return InterfacePtr->GetSelectionElement(*this, InCurrentSelection, InSelectionMethod); }
+	bool ShouldPreventTransactions() const { return InterfacePtr->ShouldPreventTransactions(*this); }
+	void WriteTransactedElement(FArchive& InArchive) const { return InterfacePtr->WriteTransactedElement(*this, InArchive); }
+	FTypedElementHandle ReadTransactedElement(FArchive& InArchive) const { return InterfacePtr->ReadTransactedElement(InArchive); }
 };
