Fixed assert when UTypedElementSelectionSet was queried for object references

#rnx

diff --git a/Engine/Source/Editor/EditorFramework/Private/Elements/Framework/TypedElementSelectionSet.cpp b/Engine/Source/Editor/EditorFramework/Private/Elements/Framework/TypedElementSelectionSet.cpp
--- a/Engine/Source/Editor/EditorFramework/Private/Elements/Framework/TypedElementSelectionSet.cpp
+++ b/Engine/Source/Editor/EditorFramework/Private/Elements/Framework/TypedElementSelectionSet.cpp
@@ -33,7 +33,7 @@ bool UTypedElementSelectionSet::Modify(bool bAlwaysMarkDirty)
 
 void UTypedElementSelectionSet::Serialize(FArchive& Ar)
 {
-	checkf(Ar.IsTransacting(), TEXT("UTypedElementSelectionSet can only be serialized for transactions!"));
+	checkf(!Ar.IsPersistent(), TEXT("UTypedElementSelectionSet can only be serialized by transient archives!"));
 
 	if (Ar.IsSaving())
 	{
