Allow the UObject utils for element selection to be used with element lists as well as the wrapper selection set

#rnx

diff --git a/Engine/Source/Editor/EditorFramework/Private/Elements/Framework/TypedElementListObjectUtil.cpp b/Engine/Source/Editor/EditorFramework/Private/Elements/Framework/TypedElementListObjectUtil.cpp
new file mode 100644
--- /dev/null
+++ b/Engine/Source/Editor/EditorFramework/Private/Elements/Framework/TypedElementListObjectUtil.cpp
@@ -0,0 +1,75 @@
+// Copyright Epic Games, Inc. All Rights Reserved.
+
+#include "Elements/Framework/TypedElementListObjectUtil.h"
+#include "Elements/Interfaces/TypedElementObjectInterface.h"
+
+namespace TypedElementListObjectUtil
+{
+
+UObject* GetObjectOfType(const TTypedElement<UTypedElementObjectInterface>& InObjectElement, const UClass* InRequiredClass)
+{
+	UObject* ElementObject = InObjectElement ? InObjectElement.GetObject() : nullptr;
+	return (ElementObject && (!InRequiredClass || ElementObject->IsA(InRequiredClass)))
+		? ElementObject
+		: nullptr;
+}
+
+void ForEachObject(const UTypedElementList* InElementList, TFunctionRef<bool(UObject*)> InCallback, const UClass* InRequiredClass)
+{
+	InElementList->ForEachElement<UTypedElementObjectInterface>([&InCallback, InRequiredClass](const TTypedElement<UTypedElementObjectInterface>& InObjectElement)
+	{
+		if (UObject* ElementObject = GetObjectOfType(InObjectElement, InRequiredClass))
+		{
+			return InCallback(ElementObject);
+		}
+		return true;
+	});
+}
+
+TArray<UObject*> GetObjects(const UTypedElementList* InElementList, const UClass* InRequiredClass)
+{
+	TArray<UObject*> SelectedObjects;
+	SelectedObjects.Reserve(InElementList->Num());
+
+	ForEachObject(InElementList, [&SelectedObjects](UObject* InObject)
+	{
+		SelectedObjects.Add(InObject);
+		return true;
+	}, InRequiredClass);
+
+	return SelectedObjects;
+}
+
+UObject* GetTopObject(const UTypedElementList* InElementList, const UClass* InRequiredClass)
+{
+	TTypedElement<UTypedElementObjectInterface> TempElement;
+	for (int32 ElementIndex = 0; ElementIndex < InElementList->Num(); ++ElementIndex)
+	{
+		InElementList->GetElementAt(ElementIndex, TempElement);
+		
+		if (UObject* ElementObject = GetObjectOfType(TempElement, InRequiredClass))
+		{
+			return ElementObject;
+		}
+	}
+
+	return nullptr;
+}
+
+UObject* GetBottomObject(const UTypedElementList* InElementList, const UClass* InRequiredClass)
+{
+	TTypedElement<UTypedElementObjectInterface> TempElement;
+	for (int32 ElementIndex = InElementList->Num() - 1; ElementIndex >= 0; --ElementIndex)
+	{
+		InElementList->GetElementAt(ElementIndex, TempElement);
+
+		if (UObject* ElementObject = GetObjectOfType(TempElement, InRequiredClass))
+		{
+			return ElementObject;
+		}
+	}
+
+	return nullptr;
+}
+
+} // namespace TypedElementListObjectUtil
diff --git a/Engine/Source/Editor/EditorFramework/Private/Elements/Framework/TypedElementSelectionSet.cpp b/Engine/Source/Editor/EditorFramework/Private/Elements/Framework/TypedElementSelectionSet.cpp
--- a/Engine/Source/Editor/EditorFramework/Private/Elements/Framework/TypedElementSelectionSet.cpp
+++ b/Engine/Source/Editor/EditorFramework/Private/Elements/Framework/TypedElementSelectionSet.cpp
@@ -2,7 +2,6 @@
 
 #include "Elements/Framework/TypedElementSelectionSet.h"
 #include "Elements/Framework/TypedElementRegistry.h"
-#include "Elements/Interfaces/TypedElementObjectInterface.h"
 
 UTypedElementSelectionSet::UTypedElementSelectionSet()
 {
@@ -211,33 +210,6 @@ FTypedElementHandle UTypedElementSelectionSet::GetSelectionElement(const FTypedE
 	return SelectionSetElement ? SelectionSetElement.GetSelectionElement(InSelectionMethod) : FTypedElementHandle();
 }
 
-TArray<UObject*> UTypedElementSelectionSet::GetSelectedObjects(const UClass* InRequiredClass) const
-{
-	TArray<UObject*> SelectedObjects;
-	SelectedObjects.Reserve(ElementList->Num());
-
-	ForEachSelectedObject([&SelectedObjects](UObject* InObject)
-	{
-		SelectedObjects.Add(InObject);
-		return true;
-	}, InRequiredClass);
-
-	return SelectedObjects;
-}
-
-void UTypedElementSelectionSet::ForEachSelectedObject(TFunctionRef<bool(UObject*)> InCallback, const UClass* InRequiredClass) const
-{
-	ElementList->ForEachElement<UTypedElementObjectInterface>([&InCallback, InRequiredClass](const TTypedElement<UTypedElementObjectInterface>& InObjectElement)
-	{
-		UObject* ElementObject = InObjectElement.GetObject();
-		if (ElementObject && (!InRequiredClass || ElementObject->IsA(InRequiredClass)))
-		{
-			return InCallback(ElementObject);
-		}
-		return true;
-	});
-}
-
 FTypedElementSelectionSetElement UTypedElementSelectionSet::ResolveSelectionSetElement(const FTypedElementHandle& InElementHandle) const
 {
 	return InElementHandle
diff --git a/Engine/Source/Editor/EditorFramework/Public/Elements/Framework/TypedElementListObjectUtil.h b/Engine/Source/Editor/EditorFramework/Public/Elements/Framework/TypedElementListObjectUtil.h
new file mode 100644
--- /dev/null
+++ b/Engine/Source/Editor/EditorFramework/Public/Elements/Framework/TypedElementListObjectUtil.h
@@ -0,0 +1,80 @@
+// Copyright Epic Games, Inc. All Rights Reserved.
+
+#pragma once
+
+#include "Elements/Framework/TypedElementList.h"
+
+namespace TypedElementListObjectUtil
+{
+
+/**
+ * Enumerate the objects from the given list of elements.
+ * @note Return true from the callback to continue enumeration.
+ */
+EDITORFRAMEWORK_API void ForEachObject(const UTypedElementList* InElementList, TFunctionRef<bool(UObject*)> InCallback, const UClass* InRequiredClass = nullptr);
+
+/**
+ * Enumerate the objects from the given list of elements.
+ * @note Return true from the callback to continue enumeration.
+ */
+template <typename RequiredClassType>
+void ForEachObject(const UTypedElementList* InElementList, TFunctionRef<bool(RequiredClassType*)> InCallback)
+{
+	ForEachObject(InElementList, [&InCallback](UObject* InObject)
+	{
+		return InCallback(CastChecked<RequiredClassType>(InObject));
+	}, RequiredClassType::StaticClass());
+}
+
+/**
+ * Get the array of objects from the given list of elements.
+ */
+EDITORFRAMEWORK_API TArray<UObject*> GetObjects(const UTypedElementList* InElementList, const UClass* InRequiredClass = nullptr);
+
+/**
+ * Get the array of objects from the given list of elements.
+ */
+template <typename RequiredClassType>
+TArray<RequiredClassType*> GetObjects(const UTypedElementList* InElementList)
+{
+	TArray<RequiredClassType*> SelectedObjects;
+	SelectedObjects.Reserve(InElementList->Num());
+
+	ForEachObject<RequiredClassType>(InElementList, [&SelectedObjects](RequiredClassType* InObject)
+	{
+		SelectedObjects.Add(InObject);
+		return true;
+	});
+
+	return SelectedObjects;
+}
+
+/**
+ * Get the first object of the given type from the given list of elements.
+ */
+EDITORFRAMEWORK_API UObject* GetTopObject(const UTypedElementList* InElementList, const UClass* InRequiredClass = nullptr);
+
+/**
+ * Get the first object of the given type from the given list of elements.
+ */
+template <typename RequiredClassType>
+RequiredClassType* GetTopObject(const UTypedElementList* InElementList)
+{
+	return Cast<RequiredClassType>(GetTopObject(InElementList, RequiredClassType::StaticClass()));
+}
+
+/**
+ * Get the last object of the given type from the given list of elements.
+ */
+EDITORFRAMEWORK_API UObject* GetBottomObject(const UTypedElementList* InElementList, const UClass* InRequiredClass = nullptr);
+
+/**
+ * Get the last object of the given type from the given list of elements.
+ */
+template <typename RequiredClassType>
+RequiredClassType* GetBottomObject(const UTypedElementList* InElementList)
+{
+	return Cast<RequiredClassType>(GetBottomObject(InElementList, RequiredClassType::StaticClass()));
+}
+
+} // namespace TypedElementListObjectUtil
diff --git a/Engine/Source/Editor/EditorFramework/Public/Elements/Framework/TypedElementSelectionSet.h b/Engine/Source/Editor/EditorFramework/Public/Elements/Framework/TypedElementSelectionSet.h
--- a/Engine/Source/Editor/EditorFramework/Public/Elements/Framework/TypedElementSelectionSet.h
+++ b/Engine/Source/Editor/EditorFramework/Public/Elements/Framework/TypedElementSelectionSet.h
@@ -3,6 +3,7 @@
 #pragma once
 
 #include "Elements/Framework/TypedElementList.h"
+#include "Elements/Framework/TypedElementListObjectUtil.h"
 #include "Elements/Framework/TypedElementAssetEditorCustomization.h"
 #include "Elements/Interfaces/TypedElementSelectionInterface.h"
 #include "TypedElementSelectionSet.generated.h"
@@ -194,11 +195,32 @@ public:
 		ElementList->ForEachElement<BaseInterfaceType>(InCallback);
 	}
 
+	/**
+	 * Get the first selected element implementing the given interface.
+	 */
+	template <typename BaseInterfaceType>
+	TTypedElement<BaseInterfaceType> GetTopSelectedElement() const
+	{
+		return ElementList->GetTopElement<BaseInterfaceType>();
+	}
+
+	/**
+	 * Get the last selected element implementing the given interface.
+	 */
+	template <typename BaseInterfaceType>
+	TTypedElement<BaseInterfaceType> GetBottomSelectedElement() const
+	{
+		return ElementList->GetBottomElement<BaseInterfaceType>();
+	}
+
 	/**
 	 * Get the array of selected objects from the currently selected elements.
 	 */
 	UFUNCTION(BlueprintCallable, BlueprintPure=false, Category="TypedElementFramework|Selection")
-	TArray<UObject*> GetSelectedObjects(const UClass* InRequiredClass = nullptr) const;
+	TArray<UObject*> GetSelectedObjects(const UClass* InRequiredClass = nullptr) const
+	{
+		return TypedElementListObjectUtil::GetObjects(ElementList, InRequiredClass);
+	}
 
 	/**
 	 * Get the array of selected objects from the currently selected elements.
@@ -206,23 +228,17 @@ public:
 	template <typename RequiredClassType>
 	TArray<RequiredClassType*> GetSelectedObjects() const
 	{
-		TArray<RequiredClassType*> SelectedObjects;
-		SelectedObjects.Reserve(ElementList->Num());
-
-		ForEachSelectedObject<RequiredClassType>([&SelectedObjects](RequiredClassType* InObject)
-		{
-			SelectedObjects.Add(InObject);
-			return true;
-		});
-
-		return SelectedObjects;
+		return TypedElementListObjectUtil::GetObjects<RequiredClassType>(ElementList);
 	}
 
 	/**
 	 * Enumerate the selected objects from the currently selected elements.
 	 * @note Return true from the callback to continue enumeration.
 	 */
-	void ForEachSelectedObject(TFunctionRef<bool(UObject*)> InCallback, const UClass* InRequiredClass = nullptr) const;
+	void ForEachSelectedObject(TFunctionRef<bool(UObject*)> InCallback, const UClass* InRequiredClass = nullptr) const
+	{
+		TypedElementListObjectUtil::ForEachObject(ElementList, InCallback, InRequiredClass);
+	}
 
 	/**
 	 * Enumerate the selected objects from the currently selected elements.
@@ -231,10 +247,43 @@ public:
 	template <typename RequiredClassType>
 	void ForEachSelectedObject(TFunctionRef<bool(RequiredClassType*)> InCallback) const
 	{
-		ForEachSelectedObject([&InCallback](UObject* InObject)
-		{
-			return InCallback(CastChecked<RequiredClassType>(InObject));
-		}, RequiredClassType::StaticClass());
+		TypedElementListObjectUtil::ForEachObject<RequiredClassType>(ElementList, InCallback);
+	}
+
+	/**
+	 * Get the first selected object of the given type.
+	 */
+	UFUNCTION(BlueprintPure, Category="TypedElementFramework|Selection")
+	UObject* GetTopSelectedObject(const UClass* InRequiredClass = nullptr) const
+	{
+		return TypedElementListObjectUtil::GetTopObject(ElementList, InRequiredClass);
+	}
+
+	/**
+	 * Get the first selected object of the given type.
+	 */
+	template <typename RequiredClassType>
+	RequiredClassType* GetTopSelectedObject() const
+	{
+		return TypedElementListObjectUtil::GetTopObject<RequiredClassType>(ElementList);
+	}
+
+	/**
+	 * Get the last selected object of the given type.
+	 */
+	UFUNCTION(BlueprintPure, Category="TypedElementFramework|Selection")
+	UObject* GetBottomSelectedObject(const UClass* InRequiredClass = nullptr) const
+	{
+		return TypedElementListObjectUtil::GetBottomObject(ElementList, InRequiredClass);
+	}
+
+	/**
+	 * Get the last selected object of the given type.
+	 */
+	template <typename RequiredClassType>
+	RequiredClassType* GetBottomSelectedObject() const
+	{
+		return TypedElementListObjectUtil::GetBottomObject<RequiredClassType>(ElementList);
 	}
 
 	/**
diff --git a/Engine/Source/Runtime/TypedElementFramework/Public/Elements/Framework/TypedElementList.h b/Engine/Source/Runtime/TypedElementFramework/Public/Elements/Framework/TypedElementList.h
--- a/Engine/Source/Runtime/TypedElementFramework/Public/Elements/Framework/TypedElementList.h
+++ b/Engine/Source/Runtime/TypedElementFramework/Public/Elements/Framework/TypedElementList.h
@@ -203,6 +203,42 @@ public:
 		TypedElementList_Private::GetElement(Registry.Get(), InElementHandle, OutElement);
 	}
 
+	/**
+	 * Get the first element implementing the given interface.
+	 */
+	template <typename BaseInterfaceType>
+	TTypedElement<BaseInterfaceType> GetTopElement() const
+	{
+		TTypedElement<BaseInterfaceType> TempElement;
+		for (int32 ElementIndex = 0; ElementIndex < Num(); ++ElementIndex)
+		{
+			GetElementAt(ElementIndex, TempElement);
+			if (TempElement)
+			{
+				break;
+			}
+		}
+		return TempElement;
+	}
+
+	/**
+	 * Get the last element implementing the given interface.
+	 */
+	template <typename BaseInterfaceType>
+	TTypedElement<BaseInterfaceType> GetBottomElement() const
+	{
+		TTypedElement<BaseInterfaceType> TempElement;
+		for (int32 ElementIndex = Num() - 1; ElementIndex >= 0; --ElementIndex)
+		{
+			GetElementAt(ElementIndex, TempElement);
+			if (TempElement)
+			{
+				break;
+			}
+		}
+		return TempElement;
+	}
+
 	/**
 	 * Get the element interface from the given handle.
 	 */
