Added utils for getting the existence/count of the number of elements/objects within a element list and selection set

#rnx

diff --git a/Engine/Source/Editor/EditorFramework/Private/Elements/Framework/TypedElementListObjectUtil.cpp b/Engine/Source/Editor/EditorFramework/Private/Elements/Framework/TypedElementListObjectUtil.cpp
--- a/Engine/Source/Editor/EditorFramework/Private/Elements/Framework/TypedElementListObjectUtil.cpp
+++ b/Engine/Source/Editor/EditorFramework/Private/Elements/Framework/TypedElementListObjectUtil.cpp
@@ -14,6 +14,32 @@ UObject* GetObjectOfType(const TTypedElement<UTypedElementObjectInterface>& InOb
 		: nullptr;
 }
 
+bool HasObjects(const UTypedElementList* InElementList, const UClass* InRequiredClass)
+{
+	bool bHasObjects = false;
+
+	ForEachObject(InElementList, [&bHasObjects](const UObject*)
+	{
+		bHasObjects = true;
+		return false;
+	}, InRequiredClass);
+
+	return bHasObjects;
+}
+
+int32 CountObjects(const UTypedElementList* InElementList, const UClass* InRequiredClass)
+{
+	int32 NumObjects = 0;
+
+	ForEachObject(InElementList, [&NumObjects](const UObject*)
+	{
+		++NumObjects;
+		return true;
+	}, InRequiredClass);
+
+	return NumObjects;
+}
+
 void ForEachObject(const UTypedElementList* InElementList, TFunctionRef<bool(UObject*)> InCallback, const UClass* InRequiredClass)
 {
 	InElementList->ForEachElement<UTypedElementObjectInterface>([&InCallback, InRequiredClass](const TTypedElement<UTypedElementObjectInterface>& InObjectElement)
diff --git a/Engine/Source/Editor/EditorFramework/Public/Elements/Framework/TypedElementListObjectUtil.h b/Engine/Source/Editor/EditorFramework/Public/Elements/Framework/TypedElementListObjectUtil.h
--- a/Engine/Source/Editor/EditorFramework/Public/Elements/Framework/TypedElementListObjectUtil.h
+++ b/Engine/Source/Editor/EditorFramework/Public/Elements/Framework/TypedElementListObjectUtil.h
@@ -7,6 +7,34 @@
 namespace TypedElementListObjectUtil
 {
 
+/**
+ * Test whether there are any objects in the given list of elements.
+ */
+EDITORFRAMEWORK_API bool HasObjects(const UTypedElementList* InElementList, const UClass* InRequiredClass = nullptr);
+
+/**
+ * Test whether there are any objects in the given list of elements.
+ */
+template <typename RequiredClassType>
+bool HasObjects(const UTypedElementList* InElementList)
+{
+	return HasObjects(InElementList, RequiredClassType::StaticClass());
+}
+
+/**
+ * Count the number of objects in the given list of elements.
+ */
+EDITORFRAMEWORK_API int32 CountObjects(const UTypedElementList* InElementList, const UClass* InRequiredClass = nullptr);
+
+/**
+ * Count the number of objects in the given list of elements.
+ */
+template <typename RequiredClassType>
+int32 CountObjects(const UTypedElementList* InElementList)
+{
+	return CountObjects(InElementList, RequiredClassType::StaticClass());
+}
+
 /**
  * Enumerate the objects from the given list of elements.
  * @note Return true from the callback to continue enumeration.
diff --git a/Engine/Source/Editor/EditorFramework/Public/Elements/Framework/TypedElementSelectionSet.h b/Engine/Source/Editor/EditorFramework/Public/Elements/Framework/TypedElementSelectionSet.h
--- a/Engine/Source/Editor/EditorFramework/Public/Elements/Framework/TypedElementSelectionSet.h
+++ b/Engine/Source/Editor/EditorFramework/Public/Elements/Framework/TypedElementSelectionSet.h
@@ -158,6 +158,33 @@ public:
 	UFUNCTION(BlueprintPure, Category="TypedElementFramework|Selection")
 	FTypedElementHandle GetSelectionElement(const FTypedElementHandle& InElementHandle, const ETypedElementSelectionMethod InSelectionMethod) const;
 
+	/**
+	 * Get the number of selected elements.
+	 */
+	UFUNCTION(BlueprintPure, Category="TypedElementFramework|Selection")
+	int32 GetNumSelectedElements() const
+	{
+		return ElementList->Num();
+	}
+
+	/**
+	 * Test whether there selected elements, optionally filtering to elements that implement the given interface.
+	 */
+	UFUNCTION(BlueprintCallable, BlueprintPure=false, Category="TypedElementFramework|Selection")
+	bool HasSelectedElements(const TSubclassOf<UTypedElementInterface>& InBaseInterfaceType = nullptr) const
+	{
+		return ElementList->HasElements(InBaseInterfaceType);
+	}
+
+	/**
+	 * Count the number of selected elements, optionally filtering to elements that implement the given interface.
+	 */
+	UFUNCTION(BlueprintCallable, BlueprintPure=false, Category="TypedElementFramework|Selection")
+	int32 CountSelectedElements(const TSubclassOf<UTypedElementInterface>& InBaseInterfaceType = nullptr) const
+	{
+		return ElementList->CountElements(InBaseInterfaceType);
+	}
+
 	/**
 	 * Get the handle of every selected element, optionally filtering to elements that implement the given interface.
 	 */
@@ -213,6 +240,42 @@ public:
 		return ElementList->GetBottomElement<BaseInterfaceType>();
 	}
 
+	/**
+	 * Test whether there are any selected objects.
+	 */
+	UFUNCTION(BlueprintCallable, BlueprintPure=false, Category="TypedElementFramework|Selection")
+	bool HasSelectedObjects(const UClass* InRequiredClass = nullptr) const
+	{
+		return TypedElementListObjectUtil::HasObjects(ElementList, InRequiredClass);
+	}
+
+	/**
+	 * Test whether there are any selected objects.
+	 */
+	template <typename RequiredClassType>
+	bool HasSelectedObjects() const
+	{
+		return TypedElementListObjectUtil::HasObjects<RequiredClassType>(ElementList);
+	}
+
+	/**
+	 * Count the number of selected objects.
+	 */
+	UFUNCTION(BlueprintCallable, BlueprintPure=false, Category="TypedElementFramework|Selection")
+	int32 CountSelectedObjects(const UClass* InRequiredClass = nullptr) const
+	{
+		return TypedElementListObjectUtil::CountObjects(ElementList, InRequiredClass);
+	}
+
+	/**
+	 * Count the number of selected objects.
+	 */
+	template <typename RequiredClassType>
+	int32 CountSelectedObjects() const
+	{
+		return TypedElementListObjectUtil::CountObjects<RequiredClassType>(ElementList);
+	}
+
 	/**
 	 * Get the array of selected objects from the currently selected elements.
 	 */
diff --git a/Engine/Source/Runtime/TypedElementFramework/Private/Elements/Framework/TypedElementList.cpp b/Engine/Source/Runtime/TypedElementFramework/Private/Elements/Framework/TypedElementList.cpp
--- a/Engine/Source/Runtime/TypedElementFramework/Private/Elements/Framework/TypedElementList.cpp
+++ b/Engine/Source/Runtime/TypedElementFramework/Private/Elements/Framework/TypedElementList.cpp
@@ -131,6 +131,46 @@ UTypedElementInterface* UTypedElementList::GetElementInterface(const FTypedEleme
 	return Registry->GetElementInterface(InElementHandle, InBaseInterfaceType);
 }
 
+bool UTypedElementList::HasElements(const TSubclassOf<UTypedElementInterface>& InBaseInterfaceType) const
+{
+	bool bHasFilteredElements = false;
+
+	if (InBaseInterfaceType)
+	{
+		ForEachElementHandle([&bHasFilteredElements](const FTypedElementHandle&)
+		{
+			bHasFilteredElements = true;
+			return false;
+		}, InBaseInterfaceType);
+	}
+	else
+	{
+		bHasFilteredElements = Num() > 0;
+	}
+
+	return bHasFilteredElements;
+}
+
+int32 UTypedElementList::CountElements(const TSubclassOf<UTypedElementInterface>& InBaseInterfaceType) const
+{
+	int32 NumFilteredElements = 0;
+
+	if (InBaseInterfaceType)
+	{
+		ForEachElementHandle([&NumFilteredElements](const FTypedElementHandle&)
+		{
+			++NumFilteredElements;
+			return true;
+		}, InBaseInterfaceType);
+	}
+	else
+	{
+		NumFilteredElements = Num();
+	}
+
+	return NumFilteredElements;
+}
+
 TArray<FTypedElementHandle> UTypedElementList::GetElementHandles(const TSubclassOf<UTypedElementInterface>& InBaseInterfaceType) const
 {
 	TArray<FTypedElementHandle> FilteredElementHandles;
diff --git a/Engine/Source/Runtime/TypedElementFramework/Public/Elements/Framework/TypedElementList.h b/Engine/Source/Runtime/TypedElementFramework/Public/Elements/Framework/TypedElementList.h
--- a/Engine/Source/Runtime/TypedElementFramework/Public/Elements/Framework/TypedElementList.h
+++ b/Engine/Source/Runtime/TypedElementFramework/Public/Elements/Framework/TypedElementList.h
@@ -254,6 +254,18 @@ public:
 	UFUNCTION(BlueprintPure, Category="TypedElementFramework|List")
 	UTypedElementInterface* GetElementInterface(const FTypedElementHandle& InElementHandle, const TSubclassOf<UTypedElementInterface>& InBaseInterfaceType) const;
 
+	/**
+	 * Test whether there are elements in this list, optionally filtering to elements that implement the given interface.
+	 */
+	UFUNCTION(BlueprintCallable, BlueprintPure=false, Category="TypedElementFramework|List")
+	bool HasElements(const TSubclassOf<UTypedElementInterface>& InBaseInterfaceType = nullptr) const;
+
+	/**
+	 * Count the number of elements in this list, optionally filtering to elements that implement the given interface.
+	 */
+	UFUNCTION(BlueprintCallable, BlueprintPure=false, Category="TypedElementFramework|List")
+	int32 CountElements(const TSubclassOf<UTypedElementInterface>& InBaseInterfaceType = nullptr) const;
+
 	/**
 	 * Get the handle of every element in this list, optionally filtering to elements that implement the given interface.
 	 */
