Fixed valid selection changes failing to transact for undo/redo

#rnx

diff --git a/Engine/Source/Editor/EditorFramework/Private/Elements/Framework/TypedElementSelectionSet.cpp b/Engine/Source/Editor/EditorFramework/Private/Elements/Framework/TypedElementSelectionSet.cpp
--- a/Engine/Source/Editor/EditorFramework/Private/Elements/Framework/TypedElementSelectionSet.cpp
+++ b/Engine/Source/Editor/EditorFramework/Private/Elements/Framework/TypedElementSelectionSet.cpp
@@ -20,7 +20,7 @@ bool UTypedElementSelectionSet::Modify(bool bAlwaysMarkDirty)
 		bool bCanModify = true;
 		ElementList->ForEachElement<UTypedElementSelectionInterface>([&bCanModify](const TTypedElement<UTypedElementSelectionInterface>& InSelectionElement)
 		{
-			bCanModify = InSelectionElement.ShouldPreventTransactions();
+			bCanModify = !InSelectionElement.ShouldPreventTransactions();
 			return bCanModify;
 		});
 
