Fixed incorrect interface signature

#rnx

diff --git a/Engine/Source/Editor/EditorFramework/Private/Elements/Framework/TypedElementViewportInteraction.cpp b/Engine/Source/Editor/EditorFramework/Private/Elements/Framework/TypedElementViewportInteraction.cpp
--- a/Engine/Source/Editor/EditorFramework/Private/Elements/Framework/TypedElementViewportInteraction.cpp
+++ b/Engine/Source/Editor/EditorFramework/Private/Elements/Framework/TypedElementViewportInteraction.cpp
@@ -21,7 +21,7 @@ bool FTypedElementAssetEditorViewportInteractionCustomization::GetGizmoPivotLoca
 	return false;
 }
 
-void FTypedElementAssetEditorViewportInteractionCustomization::PreGizmoManipulationStarted(TArrayView<const TTypedElement<UTypedElementWorldInterface>> InElementWorldHandles, const UE::Widget::EWidgetMode InWidgetMode)
+void FTypedElementAssetEditorViewportInteractionCustomization::PreGizmoManipulationStarted(TArrayView<const FTypedElementHandle> InElementHandles, const UE::Widget::EWidgetMode InWidgetMode)
 {
 }
 
@@ -78,7 +78,7 @@ void FTypedElementAssetEditorViewportInteractionCustomization::GizmoManipulation
 {
 }
 
-void FTypedElementAssetEditorViewportInteractionCustomization::PostGizmoManipulationStopped(TArrayView<const TTypedElement<UTypedElementWorldInterface>> InElementWorldHandles, const UE::Widget::EWidgetMode InWidgetMode)
+void FTypedElementAssetEditorViewportInteractionCustomization::PostGizmoManipulationStopped(TArrayView<const FTypedElementHandle> InElementHandles, const UE::Widget::EWidgetMode InWidgetMode)
 {
 }
 
@@ -101,7 +101,7 @@ void UTypedElementViewportInteraction::GetSelectedElementsToMove(const UTypedEle
 void UTypedElementViewportInteraction::BeginGizmoManipulation(const UTypedElementList* InElementsToMove, const UE::Widget::EWidgetMode InWidgetMode)
 {
 	{
-		TMap<FTypedHandleTypeId, TArray<TTypedElement<UTypedElementWorldInterface>>> ElementsToMoveByType;
+		TMap<FTypedHandleTypeId, TArray<FTypedElementHandle>> ElementsToMoveByType;
 		BatchElementsByType(InElementsToMove, ElementsToMoveByType);
 
 		for (const auto& ElementsByTypePair : ElementsToMoveByType)
@@ -151,7 +151,7 @@ void UTypedElementViewportInteraction::EndGizmoManipulation(const UTypedElementL
 	});
 	
 	{
-		TMap<FTypedHandleTypeId, TArray<TTypedElement<UTypedElementWorldInterface>>> ElementsToMoveByType;
+		TMap<FTypedHandleTypeId, TArray<FTypedElementHandle>> ElementsToMoveByType;
 		BatchElementsByType(InElementsToMove, ElementsToMoveByType);
 
 		for (const auto& ElementsByTypePair : ElementsToMoveByType)
@@ -181,13 +181,13 @@ FTypedElementViewportInteractionElement UTypedElementViewportInteraction::Resolv
 		: FTypedElementViewportInteractionElement();
 }
 
-void UTypedElementViewportInteraction::BatchElementsByType(const UTypedElementList* InElementsToMove, TMap<FTypedHandleTypeId, TArray<TTypedElement<UTypedElementWorldInterface>>>& OutElementsByType)
+void UTypedElementViewportInteraction::BatchElementsByType(const UTypedElementList* InElementsToMove, TMap<FTypedHandleTypeId, TArray<FTypedElementHandle>>& OutElementsByType)
 {
 	OutElementsByType.Reset();
-	InElementsToMove->ForEachElement<UTypedElementWorldInterface>([&OutElementsByType](const TTypedElement<UTypedElementWorldInterface>& InElementWorldHandle)
+	InElementsToMove->ForEachElementHandle([&OutElementsByType](const FTypedElementHandle& InElementHandle)
 	{
-		TArray<TTypedElement<UTypedElementWorldInterface>>& ElementsForType = OutElementsByType.FindOrAdd(InElementWorldHandle.GetId().GetTypeId());
-		ElementsForType.Add(InElementWorldHandle);
+		TArray<FTypedElementHandle>& ElementsForType = OutElementsByType.FindOrAdd(InElementHandle.GetId().GetTypeId());
+		ElementsForType.Add(InElementHandle);
 		return true;
 	});
 }
diff --git a/Engine/Source/Editor/EditorFramework/Public/Elements/Framework/TypedElementViewportInteraction.h b/Engine/Source/Editor/EditorFramework/Public/Elements/Framework/TypedElementViewportInteraction.h
--- a/Engine/Source/Editor/EditorFramework/Public/Elements/Framework/TypedElementViewportInteraction.h
+++ b/Engine/Source/Editor/EditorFramework/Public/Elements/Framework/TypedElementViewportInteraction.h
@@ -28,11 +28,11 @@ public:
 	//~ See UTypedElementViewportInteraction for API docs
 	virtual void GetElementsToMove(const TTypedElement<UTypedElementWorldInterface>& InElementWorldHandle, const ETypedElementViewportInteractionWorldType InWorldType, const UTypedElementSelectionSet* InSelectionSet, UTypedElementList* OutElementsToMove);
 	virtual bool GetGizmoPivotLocation(const TTypedElement<UTypedElementWorldInterface>& InElementWorldHandle, const UE::Widget::EWidgetMode InWidgetMode, FVector& OutPivotLocation);
-	virtual void PreGizmoManipulationStarted(TArrayView<const TTypedElement<UTypedElementWorldInterface>> InElementWorldHandles, const UE::Widget::EWidgetMode InWidgetMode);
+	virtual void PreGizmoManipulationStarted(TArrayView<const FTypedElementHandle> InElementHandles, const UE::Widget::EWidgetMode InWidgetMode);
 	virtual void GizmoManipulationStarted(const TTypedElement<UTypedElementWorldInterface>& InElementWorldHandle, const UE::Widget::EWidgetMode InWidgetMode);
 	virtual void GizmoManipulationDeltaUpdate(const TTypedElement<UTypedElementWorldInterface>& InElementWorldHandle, const UE::Widget::EWidgetMode InWidgetMode, const EAxisList::Type InDragAxis, const FInputDeviceState& InInputState, const FTransform& InDeltaTransform, const FVector& InPivotLocation);
 	virtual void GizmoManipulationStopped(const TTypedElement<UTypedElementWorldInterface>& InElementWorldHandle, const UE::Widget::EWidgetMode InWidgetMode);
-	virtual void PostGizmoManipulationStopped(TArrayView<const TTypedElement<UTypedElementWorldInterface>> InElementWorldHandles, const UE::Widget::EWidgetMode InWidgetMode);
+	virtual void PostGizmoManipulationStopped(TArrayView<const FTypedElementHandle> InElementHandles, const UE::Widget::EWidgetMode InWidgetMode);
 };
 
 /**
@@ -123,5 +123,5 @@ private:
 	/**
 	 * Batch the given elements by their type.
 	 */
-	static void BatchElementsByType(const UTypedElementList* InElementsToMove, TMap<FTypedHandleTypeId, TArray<TTypedElement<UTypedElementWorldInterface>>>& OutElementsByType);
+	static void BatchElementsByType(const UTypedElementList* InElementsToMove, TMap<FTypedHandleTypeId, TArray<FTypedElementHandle>>& OutElementsByType);
 };
diff --git a/Engine/Source/Editor/LevelEditor/Private/Elements/Actor/ActorElementLevelEditorViewportInteractionCustomization.cpp b/Engine/Source/Editor/LevelEditor/Private/Elements/Actor/ActorElementLevelEditorViewportInteractionCustomization.cpp
--- a/Engine/Source/Editor/LevelEditor/Private/Elements/Actor/ActorElementLevelEditorViewportInteractionCustomization.cpp
+++ b/Engine/Source/Editor/LevelEditor/Private/Elements/Actor/ActorElementLevelEditorViewportInteractionCustomization.cpp
@@ -99,9 +99,9 @@ void FActorElementLevelEditorViewportInteractionCustomization::GizmoManipulation
 	GEditor->BroadcastEndObjectMovement(*Actor);
 }
 
-void FActorElementLevelEditorViewportInteractionCustomization::PostGizmoManipulationStopped(TArrayView<const TTypedElement<UTypedElementWorldInterface>> InElementWorldHandles, const UE::Widget::EWidgetMode InWidgetMode)
+void FActorElementLevelEditorViewportInteractionCustomization::PostGizmoManipulationStopped(TArrayView<const FTypedElementHandle> InElementHandles, const UE::Widget::EWidgetMode InWidgetMode)
 {
-	TArray<AActor*> MovedActors = ActorElementDataUtil::GetActorsFromHandlesChecked(InElementWorldHandles);
+	TArray<AActor*> MovedActors = ActorElementDataUtil::GetActorsFromHandlesChecked(InElementHandles);
 	GEditor->BroadcastActorsMoved(MovedActors);
 }
 
diff --git a/Engine/Source/Editor/LevelEditor/Public/Elements/Actor/ActorElementLevelEditorViewportInteractionCustomization.h b/Engine/Source/Editor/LevelEditor/Public/Elements/Actor/ActorElementLevelEditorViewportInteractionCustomization.h
--- a/Engine/Source/Editor/LevelEditor/Public/Elements/Actor/ActorElementLevelEditorViewportInteractionCustomization.h
+++ b/Engine/Source/Editor/LevelEditor/Public/Elements/Actor/ActorElementLevelEditorViewportInteractionCustomization.h
@@ -15,7 +15,7 @@ public:
 	virtual void GizmoManipulationStarted(const TTypedElement<UTypedElementWorldInterface>& InElementWorldHandle, const UE::Widget::EWidgetMode InWidgetMode) override;
 	virtual void GizmoManipulationDeltaUpdate(const TTypedElement<UTypedElementWorldInterface>& InElementWorldHandle, const UE::Widget::EWidgetMode InWidgetMode, const EAxisList::Type InDragAxis, const FInputDeviceState& InInputState, const FTransform& InDeltaTransform, const FVector& InPivotLocation) override;
 	virtual void GizmoManipulationStopped(const TTypedElement<UTypedElementWorldInterface>& InElementWorldHandle, const UE::Widget::EWidgetMode InWidgetMode) override;
-	virtual void PostGizmoManipulationStopped(TArrayView<const TTypedElement<UTypedElementWorldInterface>> InElementWorldHandles, const UE::Widget::EWidgetMode InWidgetMode) override;
+	virtual void PostGizmoManipulationStopped(TArrayView<const FTypedElementHandle> InElementHandles, const UE::Widget::EWidgetMode InWidgetMode) override;
 
 	static bool CanMoveActorInViewport(const AActor* InActor, const ETypedElementViewportInteractionWorldType InWorldType);
 	static void AppendActorsToMove(AActor* InActor, const UTypedElementSelectionSet* InSelectionSet, UTypedElementList* OutElementsToMove);
