Remove redundant include

#rnx

diff --git a/Engine/Source/Editor/EditorFramework/Public/Toolkits/IToolkitHost.h b/Engine/Source/Editor/EditorFramework/Public/Toolkits/IToolkitHost.h
--- a/Engine/Source/Editor/EditorFramework/Public/Toolkits/IToolkitHost.h
+++ b/Engine/Source/Editor/EditorFramework/Public/Toolkits/IToolkitHost.h
@@ -3,7 +3,6 @@
 #pragma once
 
 #include "CoreMinimal.h"
-#include "Toolkits/IToolkit.h"
 
 class FEditorModeTools;
 
