Non-unity build fix.

#rnx

diff --git a/Engine/Source/Editor/EditorFramework/Private/Subsystems/PlacementSubsystem.cpp b/Engine/Source/Editor/EditorFramework/Private/Subsystems/PlacementSubsystem.cpp
--- a/Engine/Source/Editor/EditorFramework/Private/Subsystems/PlacementSubsystem.cpp
+++ b/Engine/Source/Editor/EditorFramework/Private/Subsystems/PlacementSubsystem.cpp
@@ -9,6 +9,7 @@
 #include "Elements/Framework/TypedElementSelectionSet.h"
 
 #include "Misc/CoreDelegates.h"
+#include "UObject/UObjectIterator.h"
 
 void UPlacementSubsystem::Initialize(FSubsystemCollectionBase& Collection)
 {
