Added SetSelection helper to UTypedElementSelectionSet

This is equivalent to calling ClearSelection then SelectElements, but happens in a single batch.

diff --git a/Engine/Source/Editor/EditorFramework/Private/Elements/Framework/TypedElementSelectionSet.cpp b/Engine/Source/Editor/EditorFramework/Private/Elements/Framework/TypedElementSelectionSet.cpp
--- a/Engine/Source/Editor/EditorFramework/Private/Elements/Framework/TypedElementSelectionSet.cpp
+++ b/Engine/Source/Editor/EditorFramework/Private/Elements/Framework/TypedElementSelectionSet.cpp
@@ -198,6 +198,23 @@ bool UTypedElementSelectionSet::ClearSelection(const FTypedElementSelectionOptio
 	return bSelectionChanged;
 }
 
+bool UTypedElementSelectionSet::SetSelection(const TArray<FTypedElementHandle>& InElementHandles, const FTypedElementSelectionOptions InSelectionOptions)
+{
+	return SetSelection(MakeArrayView(InElementHandles), InSelectionOptions);
+}
+
+bool UTypedElementSelectionSet::SetSelection(TArrayView<const FTypedElementHandle> InElementHandles, const FTypedElementSelectionOptions InSelectionOptions)
+{
+	FTypedElementListLegacySyncScopedBatch LegacySyncBatch(ElementList, InSelectionOptions.AllowLegacyNotifications());
+
+	bool bSelectionChanged = false;
+
+	bSelectionChanged |= ClearSelection(InSelectionOptions);
+	bSelectionChanged |= SelectElements(InElementHandles, InSelectionOptions);
+
+	return bSelectionChanged;
+}
+
 bool UTypedElementSelectionSet::AllowSelectionModifiers(const FTypedElementHandle& InElementHandle) const
 {
 	FTypedElementSelectionSetElement SelectionSetElement = ResolveSelectionSetElement(InElementHandle);
diff --git a/Engine/Source/Editor/EditorFramework/Public/Elements/Framework/TypedElementSelectionSet.h b/Engine/Source/Editor/EditorFramework/Public/Elements/Framework/TypedElementSelectionSet.h
--- a/Engine/Source/Editor/EditorFramework/Public/Elements/Framework/TypedElementSelectionSet.h
+++ b/Engine/Source/Editor/EditorFramework/Public/Elements/Framework/TypedElementSelectionSet.h
@@ -146,6 +146,15 @@ public:
 	UFUNCTION(BlueprintCallable, Category="TypedElementFramework|Selection")
 	bool ClearSelection(const FTypedElementSelectionOptions InSelectionOptions);
 
+	/**
+	 * Attempt to make the selection the given elements.
+	 * @note Equivalent to calling ClearSelection then SelectElements, but happens in a single batch.
+	 * @return True if the selection was changed, false otherwise.
+	 */
+	UFUNCTION(BlueprintCallable, Category="TypedElementFramework|Selection")
+	bool SetSelection(const TArray<FTypedElementHandle>& InElementHandles, const FTypedElementSelectionOptions InSelectionOptions);
+	bool SetSelection(TArrayView<const FTypedElementHandle> InElementHandles, const FTypedElementSelectionOptions InSelectionOptions);
+
 	/**
 	 * Test to see whether selection modifiers (Ctrl or Shift) are allowed while selecting this element.
 	 */
diff --git a/Engine/Source/Editor/UnrealEd/Private/LevelEditorViewport.cpp b/Engine/Source/Editor/UnrealEd/Private/LevelEditorViewport.cpp
--- a/Engine/Source/Editor/UnrealEd/Private/LevelEditorViewport.cpp
+++ b/Engine/Source/Editor/UnrealEd/Private/LevelEditorViewport.cpp
@@ -2597,10 +2597,7 @@ bool FLevelEditorViewportClient::InputWidgetDelta(FViewport* InViewport, EAxisLi
 									const FTypedElementSelectionOptions SelectionOptions = FTypedElementSelectionOptions()
 										.SetAllowLegacyNotifications(false); // Old drag duplicate code didn't used to notify about this selection change
 
-									FTypedElementListLegacySyncScopedBatch LegacySyncBatch(SelectionSet->GetElementList(), SelectionOptions.AllowLegacyNotifications());
-
-									SelectionSet->ClearSelection(SelectionOptions);
-									SelectionSet->SelectElements(DuplicatedElements, SelectionOptions);
+									SelectionSet->SetSelection(DuplicatedElements, SelectionOptions);
 								}
 							}
 
