Added support for deleting typed elements from a world

diff --git a/Engine/Source/Editor/EditorFramework/Private/Elements/Framework/TypedElementCommonActions.cpp b/Engine/Source/Editor/EditorFramework/Private/Elements/Framework/TypedElementCommonActions.cpp
new file mode 100644
--- /dev/null
+++ b/Engine/Source/Editor/EditorFramework/Private/Elements/Framework/TypedElementCommonActions.cpp
@@ -0,0 +1,137 @@
+// Copyright Epic Games, Inc. All Rights Reserved.
+
+#include "Elements/Framework/TypedElementCommonActions.h"
+#include "Elements/Framework/TypedElementSelectionSet.h"
+#include "Elements/Framework/TypedElementRegistry.h"
+#include "Elements/Framework/TypedElementUtil.h"
+
+void FTypedElementCommonActionsCustomization::GetElementsToDelete(const TTypedElement<UTypedElementWorldInterface>& InElementWorldHandle, const UTypedElementSelectionSet* InSelectionSet, UTypedElementList* OutElementsToDelete)
+{
+	OutElementsToDelete->Add(InElementWorldHandle);
+}
+
+bool FTypedElementCommonActionsCustomization::DeleteElements(UTypedElementWorldInterface* InWorldInterface, TArrayView<const FTypedElementHandle> InElementHandles, UWorld* InWorld, UTypedElementSelectionSet* InSelectionSet, const FTypedElementDeletionOptions& InDeletionOptions)
+{
+	return InWorldInterface->DeleteElements(InElementHandles, InWorld, InSelectionSet, InDeletionOptions);
+}
+
+void FTypedElementCommonActionsCustomization::DuplicateElements(UTypedElementWorldInterface* InWorldInterface, TArrayView<const FTypedElementHandle> InElementHandles, UWorld* InWorld, bool bOffsetLocations, TArray<FTypedElementHandle>& OutNewElements)
+{
+	InWorldInterface->DuplicateElements(InElementHandles, InWorld, bOffsetLocations, OutNewElements);
+}
+
+
+void UTypedElementCommonActions::GetSelectedElementsToDelete(const UTypedElementSelectionSet* InSelectionSet, UTypedElementList* OutElementsToDelete) const
+{
+	OutElementsToDelete->Reset();
+	InSelectionSet->ForEachSelectedElement<UTypedElementWorldInterface>([this, InSelectionSet, OutElementsToDelete](const TTypedElement<UTypedElementWorldInterface>& InElementWorldHandle)
+	{
+		FTypedElementCommonActionsElement CommonActionsElement(InElementWorldHandle, GetInterfaceCustomizationByTypeId(InElementWorldHandle.GetId().GetTypeId()));
+		check(CommonActionsElement.IsSet());
+		CommonActionsElement.GetElementsToDelete(InSelectionSet, OutElementsToDelete);
+		return true;
+	});
+}
+
+bool UTypedElementCommonActions::DeleteElements(const TArray<FTypedElementHandle>& ElementHandles, UWorld* World, UTypedElementSelectionSet* InSelectionSet, const FTypedElementDeletionOptions& InDeletionOptions)
+{
+	return DeleteElements(MakeArrayView(ElementHandles), World, InSelectionSet, InDeletionOptions);
+}
+
+bool UTypedElementCommonActions::DeleteElements(TArrayView<const FTypedElementHandle> ElementHandles, UWorld* World, UTypedElementSelectionSet* InSelectionSet, const FTypedElementDeletionOptions& InDeletionOptions)
+{
+	TMap<FTypedHandleTypeId, TArray<FTypedElementHandle>> ElementsToDuplicateByType;
+	TypedElementUtil::BatchElementsByType(ElementHandles, ElementsToDuplicateByType);
+
+	bool bSuccess = false;
+
+	UTypedElementRegistry* Registry = UTypedElementRegistry::GetInstance();
+	for (const auto& ElementsByTypePair : ElementsToDuplicateByType)
+	{
+		FTypedElementCommonActionsCustomization* CommonActionsCustomization = GetInterfaceCustomizationByTypeId(ElementsByTypePair.Key);
+		UTypedElementWorldInterface* WorldInterface = Registry->GetElementInterface<UTypedElementWorldInterface>(ElementsByTypePair.Key);
+		if (CommonActionsCustomization && WorldInterface)
+		{
+			bSuccess |= CommonActionsCustomization->DeleteElements(WorldInterface, ElementsByTypePair.Value, World, InSelectionSet, InDeletionOptions);
+		}
+	}
+
+	return bSuccess;
+}
+
+bool UTypedElementCommonActions::DeleteElements(const UTypedElementList* ElementList, UWorld* World, UTypedElementSelectionSet* InSelectionSet, const FTypedElementDeletionOptions& InDeletionOptions)
+{
+	TMap<FTypedHandleTypeId, TArray<FTypedElementHandle>> ElementsToDuplicateByType;
+	TypedElementUtil::BatchElementsByType(ElementList, ElementsToDuplicateByType);
+
+	bool bSuccess = false;
+
+	UTypedElementRegistry* Registry = UTypedElementRegistry::GetInstance();
+	for (const auto& ElementsByTypePair : ElementsToDuplicateByType)
+	{
+		FTypedElementCommonActionsCustomization* CommonActionsCustomization = GetInterfaceCustomizationByTypeId(ElementsByTypePair.Key);
+		UTypedElementWorldInterface* WorldInterface = Registry->GetElementInterface<UTypedElementWorldInterface>(ElementsByTypePair.Key);
+		if (CommonActionsCustomization && WorldInterface)
+		{
+			bSuccess |= CommonActionsCustomization->DeleteElements(WorldInterface, ElementsByTypePair.Value, World, InSelectionSet, InDeletionOptions);
+		}
+	}
+
+	return bSuccess;
+}
+
+TArray<FTypedElementHandle> UTypedElementCommonActions::DuplicateElements(const TArray<FTypedElementHandle>& ElementHandles, UWorld* World, bool bOffsetLocations)
+{
+	return DuplicateElements(MakeArrayView(ElementHandles), World, bOffsetLocations);
+}
+
+TArray<FTypedElementHandle> UTypedElementCommonActions::DuplicateElements(TArrayView<const FTypedElementHandle> ElementHandles, UWorld* World, bool bOffsetLocations)
+{
+	TMap<FTypedHandleTypeId, TArray<FTypedElementHandle>> ElementsToDuplicateByType;
+	TypedElementUtil::BatchElementsByType(ElementHandles, ElementsToDuplicateByType);
+
+	TArray<FTypedElementHandle> NewElements;
+	NewElements.Reserve(ElementHandles.Num());
+
+	UTypedElementRegistry* Registry = UTypedElementRegistry::GetInstance();
+	for (const auto& ElementsByTypePair : ElementsToDuplicateByType)
+	{
+		FTypedElementCommonActionsCustomization* CommonActionsCustomization = GetInterfaceCustomizationByTypeId(ElementsByTypePair.Key);
+		UTypedElementWorldInterface* WorldInterface = Registry->GetElementInterface<UTypedElementWorldInterface>(ElementsByTypePair.Key);
+		if (CommonActionsCustomization && WorldInterface)
+		{
+			CommonActionsCustomization->DuplicateElements(WorldInterface, ElementsByTypePair.Value, World, bOffsetLocations, NewElements);
+		}
+	}
+
+	return NewElements;
+}
+
+TArray<FTypedElementHandle> UTypedElementCommonActions::DuplicateElements(const UTypedElementList* ElementList, UWorld* World, bool bOffsetLocations)
+{
+	TMap<FTypedHandleTypeId, TArray<FTypedElementHandle>> ElementsToDuplicateByType;
+	TypedElementUtil::BatchElementsByType(ElementList, ElementsToDuplicateByType);
+
+	TArray<FTypedElementHandle> NewElements;
+	NewElements.Reserve(ElementList->Num());
+
+	UTypedElementRegistry* Registry = UTypedElementRegistry::GetInstance();
+	for (const auto& ElementsByTypePair : ElementsToDuplicateByType)
+	{
+		FTypedElementCommonActionsCustomization* CommonActionsCustomization = GetInterfaceCustomizationByTypeId(ElementsByTypePair.Key);
+		UTypedElementWorldInterface* WorldInterface = Registry->GetElementInterface<UTypedElementWorldInterface>(ElementsByTypePair.Key);
+		if (CommonActionsCustomization && WorldInterface)
+		{
+			CommonActionsCustomization->DuplicateElements(WorldInterface, ElementsByTypePair.Value, World, bOffsetLocations, NewElements);
+		}
+	}
+
+	return NewElements;
+}
+
+FTypedElementCommonActionsElement UTypedElementCommonActions::ResolveCommonActionsElement(const FTypedElementHandle& InElementHandle) const
+{
+	return InElementHandle
+		? FTypedElementCommonActionsElement(UTypedElementRegistry::GetInstance()->GetElement<UTypedElementWorldInterface>(InElementHandle), GetInterfaceCustomizationByTypeId(InElementHandle.GetId().GetTypeId()))
+		: FTypedElementCommonActionsElement();
+}
diff --git a/Engine/Source/Editor/EditorFramework/Public/Elements/Framework/TypedElementCommonActions.h b/Engine/Source/Editor/EditorFramework/Public/Elements/Framework/TypedElementCommonActions.h
new file mode 100644
--- /dev/null
+++ b/Engine/Source/Editor/EditorFramework/Public/Elements/Framework/TypedElementCommonActions.h
@@ -0,0 +1,95 @@
+// Copyright Epic Games, Inc. All Rights Reserved.
+
+#pragma once
+
+#include "Elements/Framework/TypedElementInterfaceCustomization.h"
+#include "Elements/Interfaces/TypedElementWorldInterface.h"
+#include "TypedElementCommonActions.generated.h"
+
+class UTypedElementList;
+class UTypedElementSelectionSet;
+
+/**
+ * Customization used to allow asset editors (such as the level editor) to override the base behavior of common actions.
+ */
+class EDITORFRAMEWORK_API FTypedElementCommonActionsCustomization
+{
+public:
+	virtual ~FTypedElementCommonActionsCustomization() = default;
+
+	//~ See UTypedElementCommonActions for API docs
+	virtual void GetElementsToDelete(const TTypedElement<UTypedElementWorldInterface>& InElementWorldHandle, const UTypedElementSelectionSet* InSelectionSet, UTypedElementList* OutElementsToDelete);
+	virtual bool DeleteElements(UTypedElementWorldInterface* InWorldInterface, TArrayView<const FTypedElementHandle> InElementHandles, UWorld* InWorld, UTypedElementSelectionSet* InSelectionSet, const FTypedElementDeletionOptions& InDeletionOptions);
+	virtual void DuplicateElements(UTypedElementWorldInterface* InWorldInterface, TArrayView<const FTypedElementHandle> InElementHandles, UWorld* InWorld, bool bOffsetLocations, TArray<FTypedElementHandle>& OutNewElements);
+};
+
+/**
+ * Utility to hold a typed element handle and its associated world interface and common actions customization.
+ */
+struct EDITORFRAMEWORK_API FTypedElementCommonActionsElement
+{
+public:
+	FTypedElementCommonActionsElement() = default;
+
+	FTypedElementCommonActionsElement(TTypedElement<UTypedElementWorldInterface> InElementWorldHandle, FTypedElementCommonActionsCustomization* InCommonActionsCustomization)
+		: ElementWorldHandle(MoveTemp(InElementWorldHandle))
+		, CommonActionsCustomization(InCommonActionsCustomization)
+	{
+	}
+
+	FTypedElementCommonActionsElement(const FTypedElementCommonActionsElement&) = default;
+	FTypedElementCommonActionsElement& operator=(const FTypedElementCommonActionsElement&) = default;
+
+	FTypedElementCommonActionsElement(FTypedElementCommonActionsElement&&) = default;
+	FTypedElementCommonActionsElement& operator=(FTypedElementCommonActionsElement&&) = default;
+
+	FORCEINLINE explicit operator bool() const
+	{
+		return IsSet();
+	}
+
+	FORCEINLINE bool IsSet() const
+	{
+		return ElementWorldHandle.IsSet()
+			&& CommonActionsCustomization;
+	}
+
+	//~ See UTypedElementCommonActions for API docs
+	void GetElementsToDelete(const UTypedElementSelectionSet* InSelectionSet, UTypedElementList* OutElementsToDelete) { CommonActionsCustomization->GetElementsToDelete(ElementWorldHandle, InSelectionSet, OutElementsToDelete); }
+
+private:
+	TTypedElement<UTypedElementWorldInterface> ElementWorldHandle;
+	FTypedElementCommonActionsCustomization* CommonActionsCustomization = nullptr;
+};
+
+/**
+ * A utility to handle higher-level common actions, but default via UTypedElementWorldInterface,
+ * but asset editors can customize this behavior via FTypedElementCommonActionsCustomization.
+ */
+UCLASS(Transient)
+class EDITORFRAMEWORK_API UTypedElementCommonActions : public UObject, public TTypedElementInterfaceCustomizationRegistry<FTypedElementCommonActionsCustomization>
+{
+	GENERATED_BODY()
+
+public:
+	/**
+	 * Get the elements from the given selection set that should be deleted
+	 */
+	void GetSelectedElementsToDelete(const UTypedElementSelectionSet* InSelectionSet, UTypedElementList* OutElementsToDelete) const;
+
+	UFUNCTION(BlueprintCallable, Category="TypedElementFramework|Common")
+	bool DeleteElements(const TArray<FTypedElementHandle>& ElementHandles, UWorld* World, UTypedElementSelectionSet* InSelectionSet, const FTypedElementDeletionOptions& InDeletionOptions);
+	bool DeleteElements(TArrayView<const FTypedElementHandle> ElementHandles, UWorld* World, UTypedElementSelectionSet* InSelectionSet, const FTypedElementDeletionOptions& InDeletionOptions);
+	bool DeleteElements(const UTypedElementList* ElementList, UWorld* World, UTypedElementSelectionSet* InSelectionSet, const FTypedElementDeletionOptions& InDeletionOptions);
+
+	UFUNCTION(BlueprintCallable, Category="TypedElementFramework|Common")
+	TArray<FTypedElementHandle> DuplicateElements(const TArray<FTypedElementHandle>& ElementHandles, UWorld* World, bool bOffsetLocations);
+	TArray<FTypedElementHandle> DuplicateElements(TArrayView<const FTypedElementHandle> ElementHandles, UWorld* World, bool bOffsetLocations);
+	TArray<FTypedElementHandle> DuplicateElements(const UTypedElementList* ElementList, UWorld* World, bool bOffsetLocations);
+
+private:
+	/**
+	 * Attempt to resolve the selection interface and common actions customization for the given element, if any.
+	 */
+	FTypedElementCommonActionsElement ResolveCommonActionsElement(const FTypedElementHandle& InElementHandle) const;
+};
diff --git a/Engine/Source/Editor/EditorFramework/Public/Toolkits/IToolkitHost.h b/Engine/Source/Editor/EditorFramework/Public/Toolkits/IToolkitHost.h
--- a/Engine/Source/Editor/EditorFramework/Public/Toolkits/IToolkitHost.h
+++ b/Engine/Source/Editor/EditorFramework/Public/Toolkits/IToolkitHost.h
@@ -5,6 +5,7 @@
 #include "CoreMinimal.h"
 
 class FEditorModeTools;
+class UTypedElementCommonActions;
 
 /**
  * Base interface class for toolkit hosts
@@ -35,5 +36,7 @@ public:
 
 	/** Returns the mode manager for this toolkit host. For standalone toolkits */
 	virtual FEditorModeTools& GetEditorModeManager() const = 0;
-};
 
+	/** Returns the common actions implementation for this toolkit host */
+	virtual UTypedElementCommonActions* GetCommonActions() const = 0;
+};
diff --git a/Engine/Source/Editor/LevelEditor/Private/Elements/Actor/ActorElementLevelEditorCommonActionsCustomization.cpp b/Engine/Source/Editor/LevelEditor/Private/Elements/Actor/ActorElementLevelEditorCommonActionsCustomization.cpp
new file mode 100644
--- /dev/null
+++ b/Engine/Source/Editor/LevelEditor/Private/Elements/Actor/ActorElementLevelEditorCommonActionsCustomization.cpp
@@ -0,0 +1,52 @@
+// Copyright Epic Games, Inc. All Rights Reserved.
+
+#include "Elements/Actor/ActorElementLevelEditorCommonActionsCustomization.h"
+#include "Elements/Actor/ActorElementData.h"
+#include "GameFramework/Actor.h"
+
+#include "Elements/Framework/TypedElementSelectionSet.h"
+#include "Elements/Component/ComponentElementSelectionInterface.h"
+
+#include "EditorModeManager.h"
+#include "Toolkits/IToolkitHost.h"
+
+void FActorElementLevelEditorCommonActionsCustomization::GetElementsToDelete(const TTypedElement<UTypedElementWorldInterface>& InElementWorldHandle, const UTypedElementSelectionSet* InSelectionSet, UTypedElementList* OutElementsToDelete)
+{
+	AActor* Actor = ActorElementDataUtil::GetActorFromHandleChecked(InElementWorldHandle);
+
+	if (UComponentElementSelectionInterface::HasSelectedComponents(InSelectionSet->GetElementList()))
+	{
+		// If we have components selected then we will delete those rather than the actors
+		return;
+	}
+
+	FTypedElementCommonActionsCustomization::GetElementsToDelete(InElementWorldHandle, InSelectionSet, OutElementsToDelete);
+}
+
+bool FActorElementLevelEditorCommonActionsCustomization::DeleteElements(UTypedElementWorldInterface* InWorldInterface, TArrayView<const FTypedElementHandle> InElementHandles, UWorld* InWorld, UTypedElementSelectionSet* InSelectionSet, const FTypedElementDeletionOptions& InDeletionOptions)
+{
+	if (const IToolkitHost* ToolkitHostPtr = GetToolkitHost())
+	{
+		// TODO: Needs to pass in the actors to delete
+		if (ToolkitHostPtr->GetEditorModeManager().ProcessEditDelete())
+		{
+			return true;
+		}
+	}
+
+	return FTypedElementCommonActionsCustomization::DeleteElements(InWorldInterface, InElementHandles, InWorld, InSelectionSet, InDeletionOptions);
+}
+
+void FActorElementLevelEditorCommonActionsCustomization::DuplicateElements(UTypedElementWorldInterface* InWorldInterface, TArrayView<const FTypedElementHandle> InElementHandles, UWorld* InWorld, bool bOffsetLocations, TArray<FTypedElementHandle>& OutNewElements)
+{
+	if (const IToolkitHost* ToolkitHostPtr = GetToolkitHost())
+	{
+		// TODO: Needs to pass in the actors to duplicate
+		if (ToolkitHostPtr->GetEditorModeManager().ProcessEditDuplicate())
+		{
+			return;
+		}
+	}
+
+	FTypedElementCommonActionsCustomization::DuplicateElements(InWorldInterface, InElementHandles, InWorld, bOffsetLocations, OutNewElements);
+}
diff --git a/Engine/Source/Editor/LevelEditor/Private/Elements/Actor/ActorElementLevelEditorCommonActionsCustomization.h b/Engine/Source/Editor/LevelEditor/Private/Elements/Actor/ActorElementLevelEditorCommonActionsCustomization.h
new file mode 100644
--- /dev/null
+++ b/Engine/Source/Editor/LevelEditor/Private/Elements/Actor/ActorElementLevelEditorCommonActionsCustomization.h
@@ -0,0 +1,14 @@
+// Copyright Epic Games, Inc. All Rights Reserved.
+
+#pragma once
+
+#include "Elements/Framework/TypedElementCommonActions.h"
+#include "Elements/Framework/TypedElementAssetEditorToolkitHostMixin.h"
+
+class FActorElementLevelEditorCommonActionsCustomization : public FTypedElementCommonActionsCustomization, public FTypedElementAssetEditorToolkitHostMixin
+{
+public:
+	virtual void GetElementsToDelete(const TTypedElement<UTypedElementWorldInterface>& InElementWorldHandle, const UTypedElementSelectionSet* InSelectionSet, UTypedElementList* OutElementsToDelete) override;
+	virtual bool DeleteElements(UTypedElementWorldInterface* InWorldInterface, TArrayView<const FTypedElementHandle> InElementHandles, UWorld* InWorld, UTypedElementSelectionSet* InSelectionSet, const FTypedElementDeletionOptions& InDeletionOptions) override;
+	virtual void DuplicateElements(UTypedElementWorldInterface* InWorldInterface, TArrayView<const FTypedElementHandle> InElementHandles, UWorld* InWorld, bool bOffsetLocations, TArray<FTypedElementHandle>& OutNewElements) override;
+};
diff --git a/Engine/Source/Editor/LevelEditor/Private/Elements/Component/ComponentElementLevelEditorCommonActionsCustomization.cpp b/Engine/Source/Editor/LevelEditor/Private/Elements/Component/ComponentElementLevelEditorCommonActionsCustomization.cpp
new file mode 100644
--- /dev/null
+++ b/Engine/Source/Editor/LevelEditor/Private/Elements/Component/ComponentElementLevelEditorCommonActionsCustomization.cpp
@@ -0,0 +1,36 @@
+// Copyright Epic Games, Inc. All Rights Reserved.
+
+#include "Elements/Component/ComponentElementLevelEditorCommonActionsCustomization.h"
+#include "Elements/Component/ComponentElementData.h"
+#include "Components/ActorComponent.h"
+
+#include "EditorModeManager.h"
+#include "Toolkits/IToolkitHost.h"
+
+bool FComponentElementLevelEditorCommonActionsCustomization::DeleteElements(UTypedElementWorldInterface* InWorldInterface, TArrayView<const FTypedElementHandle> InElementHandles, UWorld* InWorld, UTypedElementSelectionSet* InSelectionSet, const FTypedElementDeletionOptions& InDeletionOptions)
+{
+	if (const IToolkitHost* ToolkitHostPtr = GetToolkitHost())
+	{
+		// TODO: Needs to pass in the components to delete
+		if (ToolkitHostPtr->GetEditorModeManager().ProcessEditDelete())
+		{
+			return true;
+		}
+	}
+
+	return FTypedElementCommonActionsCustomization::DeleteElements(InWorldInterface, InElementHandles, InWorld, InSelectionSet, InDeletionOptions);
+}
+
+void FComponentElementLevelEditorCommonActionsCustomization::DuplicateElements(UTypedElementWorldInterface* InWorldInterface, TArrayView<const FTypedElementHandle> InElementHandles, UWorld* InWorld, bool bOffsetLocations, TArray<FTypedElementHandle>& OutNewElements)
+{
+	if (const IToolkitHost* ToolkitHostPtr = GetToolkitHost())
+	{
+		// TODO: Needs to pass in the components to duplicate
+		if (ToolkitHostPtr->GetEditorModeManager().ProcessEditDuplicate())
+		{
+			return;
+		}
+	}
+
+	FTypedElementCommonActionsCustomization::DuplicateElements(InWorldInterface, InElementHandles, InWorld, bOffsetLocations, OutNewElements);
+}
diff --git a/Engine/Source/Editor/LevelEditor/Private/Elements/Component/ComponentElementLevelEditorCommonActionsCustomization.h b/Engine/Source/Editor/LevelEditor/Private/Elements/Component/ComponentElementLevelEditorCommonActionsCustomization.h
new file mode 100644
--- /dev/null
+++ b/Engine/Source/Editor/LevelEditor/Private/Elements/Component/ComponentElementLevelEditorCommonActionsCustomization.h
@@ -0,0 +1,13 @@
+// Copyright Epic Games, Inc. All Rights Reserved.
+
+#pragma once
+
+#include "Elements/Framework/TypedElementCommonActions.h"
+#include "Elements/Framework/TypedElementAssetEditorToolkitHostMixin.h"
+
+class FComponentElementLevelEditorCommonActionsCustomization : public FTypedElementCommonActionsCustomization, public FTypedElementAssetEditorToolkitHostMixin
+{
+public:
+	virtual bool DeleteElements(UTypedElementWorldInterface* InWorldInterface, TArrayView<const FTypedElementHandle> InElementHandles, UWorld* InWorld, UTypedElementSelectionSet* InSelectionSet, const FTypedElementDeletionOptions& InDeletionOptions) override;
+	virtual void DuplicateElements(UTypedElementWorldInterface* InWorldInterface, TArrayView<const FTypedElementHandle> InElementHandles, UWorld* InWorld, bool bOffsetLocations, TArray<FTypedElementHandle>& OutNewElements) override;
+};
diff --git a/Engine/Source/Editor/LevelEditor/Private/SLevelEditor.cpp b/Engine/Source/Editor/LevelEditor/Private/SLevelEditor.cpp
--- a/Engine/Source/Editor/LevelEditor/Private/SLevelEditor.cpp
+++ b/Engine/Source/Editor/LevelEditor/Private/SLevelEditor.cpp
@@ -61,8 +61,11 @@
 #include "Editor/EnvironmentLightingViewer/Public/EnvironmentLightingModule.h"
 #include "Misc/MessageDialog.h"
 #include "Elements/Framework/TypedElementSelectionSet.h"
+#include "Elements/Framework/TypedElementCommonActions.h"
 #include "Elements/Actor/ActorElementLevelEditorSelectionCustomization.h"
+#include "Elements/Actor/ActorElementLevelEditorCommonActionsCustomization.h"
 #include "Elements/Component/ComponentElementLevelEditorSelectionCustomization.h"
+#include "Elements/Component/ComponentElementLevelEditorCommonActionsCustomization.h"
 #include "ThumbnailRendering/ThumbnailManager.h"
 
 #define LOCTEXT_NAMESPACE "SLevelEditor"
@@ -244,6 +247,21 @@ void SLevelEditor::Initialize( const TSharedRef<SDockTab>& OwnerTab, const TShar
 	GUnrealEd->GetSelectedActors()->SetElementSelectionSet(SelectedElements);
 	GUnrealEd->GetSelectedComponents()->SetElementSelectionSet(SelectedElements);
 
+	CommonActions = NewObject<UTypedElementCommonActions>();
+	CommonActions->AddToRoot();
+
+	// Register the level editor specific selection behavior
+	{
+		TUniquePtr<FActorElementLevelEditorCommonActionsCustomization> ActorCustomization = MakeUnique<FActorElementLevelEditorCommonActionsCustomization>();
+		ActorCustomization->SetToolkitHost(this);
+		CommonActions->RegisterInterfaceCustomizationByTypeName(NAME_Actor, MoveTemp(ActorCustomization));
+	}
+	{
+		TUniquePtr<FComponentElementLevelEditorCommonActionsCustomization> ComponentCustomization = MakeUnique<FComponentElementLevelEditorCommonActionsCustomization>();
+		ComponentCustomization->SetToolkitHost(this);
+		CommonActions->RegisterInterfaceCustomizationByTypeName(NAME_Components, MoveTemp(ComponentCustomization));
+	}
+
 	// Bind the level editor tab's label to the currently loaded level name string in the main frame
 	OwnerTab->SetLabel( TAttribute<FText>( this, &SLevelEditor::GetTabTitle) );
 	OwnerTab->SetTabLabelSuffix(TAttribute<FText>(this, &SLevelEditor::GetTabSuffix));
@@ -364,6 +382,9 @@ SLevelEditor::~SLevelEditor()
 		}
 		SelectedElements->RemoveFromRoot();
 		SelectedElements = nullptr;
+
+		CommonActions->RemoveFromRoot();
+		CommonActions = nullptr;
 	}
 }
 
@@ -1979,4 +2000,9 @@ FEditorModeTools& SLevelEditor::GetEditorModeManager() const
 	return GLevelEditorModeTools();
 }
 
+UTypedElementCommonActions* SLevelEditor::GetCommonActions() const
+{
+	return CommonActions;
+}
+
 #undef LOCTEXT_NAMESPACE
\ No newline at end of file
diff --git a/Engine/Source/Editor/LevelEditor/Private/SLevelEditor.h b/Engine/Source/Editor/LevelEditor/Private/SLevelEditor.h
--- a/Engine/Source/Editor/LevelEditor/Private/SLevelEditor.h
+++ b/Engine/Source/Editor/LevelEditor/Private/SLevelEditor.h
@@ -129,6 +129,7 @@ public:
 	virtual void SetActorDetailsSCSEditorUICustomization(TSharedPtr<ISCSEditorUICustomization> InActorDetailsSCSEditorUICustomization) override;
 	virtual TSharedRef<SWidget> CreateToolBox() override;
 	virtual FEditorModeTools& GetEditorModeManager() const override;
+	virtual UTypedElementCommonActions* GetCommonActions() const override;
 
 	/** SWidget overrides */
 	virtual bool SupportsKeyboardFocus() const override
@@ -275,6 +276,9 @@ private:
 	// The list of selected elements (also set on the global USelection for actors and components).
 	UTypedElementSelectionSet* SelectedElements = nullptr;
 
+	// The common actions implementation for the level editor
+	UTypedElementCommonActions* CommonActions = nullptr;
+
 	// The box that holds the title bar messages.
 	TSharedPtr<SHorizontalBox> TtileBarMessageBox;
 
diff --git a/Engine/Source/Editor/UnrealEd/Private/Elements/Actor/ActorElementEditorWorldInterface.cpp b/Engine/Source/Editor/UnrealEd/Private/Elements/Actor/ActorElementEditorWorldInterface.cpp
--- a/Engine/Source/Editor/UnrealEd/Private/Elements/Actor/ActorElementEditorWorldInterface.cpp
+++ b/Engine/Source/Editor/UnrealEd/Private/Elements/Actor/ActorElementEditorWorldInterface.cpp
@@ -61,6 +61,13 @@ void UActorElementEditorWorldInterface::NotifyMovementEnded(const FTypedElementH
 	}
 }
 
+bool UActorElementEditorWorldInterface::DeleteElements(TArrayView<const FTypedElementHandle> InElementHandles, UWorld* InWorld, UTypedElementSelectionSet* InSelectionSet, const FTypedElementDeletionOptions& InDeletionOptions)
+{
+	const TArray<AActor*> ActorsToDelete = ActorElementDataUtil::GetActorsFromHandles(InElementHandles);
+	return ActorsToDelete.Num() > 0
+		&& GUnrealEd->DeleteActors(ActorsToDelete, InWorld, InSelectionSet, InDeletionOptions.VerifyDeletionCanHappen(), InDeletionOptions.WarnAboutReferences(), InDeletionOptions.WarnAboutSoftReferences());
+}
+
 void UActorElementEditorWorldInterface::DuplicateElements(TArrayView<const FTypedElementHandle> InElementHandles, UWorld* InWorld, bool bOffsetLocations, TArray<FTypedElementHandle>& OutNewElements)
 {
 	const TArray<AActor*> ActorsToDuplicate = ActorElementDataUtil::GetActorsFromHandles(InElementHandles);
diff --git a/Engine/Source/Editor/UnrealEd/Private/Elements/Actor/ActorElementEditorWorldInterface.h b/Engine/Source/Editor/UnrealEd/Private/Elements/Actor/ActorElementEditorWorldInterface.h
--- a/Engine/Source/Editor/UnrealEd/Private/Elements/Actor/ActorElementEditorWorldInterface.h
+++ b/Engine/Source/Editor/UnrealEd/Private/Elements/Actor/ActorElementEditorWorldInterface.h
@@ -16,5 +16,6 @@ public:
 	virtual void NotifyMovementStarted(const FTypedElementHandle& InElementHandle) override;
 	virtual void NotifyMovementOngoing(const FTypedElementHandle& InElementHandle) override;
 	virtual void NotifyMovementEnded(const FTypedElementHandle& InElementHandle) override;
+	virtual bool DeleteElements(TArrayView<const FTypedElementHandle> InElementHandles, UWorld* InWorld, UTypedElementSelectionSet* InSelectionSet, const FTypedElementDeletionOptions& InDeletionOptions) override;
 	virtual void DuplicateElements(TArrayView<const FTypedElementHandle> InElementHandles, UWorld* InWorld, bool bOffsetLocations, TArray<FTypedElementHandle>& OutNewElements) override;
 };
diff --git a/Engine/Source/Editor/UnrealEd/Private/Elements/Component/ComponentElementEditorWorldInterface.cpp b/Engine/Source/Editor/UnrealEd/Private/Elements/Component/ComponentElementEditorWorldInterface.cpp
--- a/Engine/Source/Editor/UnrealEd/Private/Elements/Component/ComponentElementEditorWorldInterface.cpp
+++ b/Engine/Source/Editor/UnrealEd/Private/Elements/Component/ComponentElementEditorWorldInterface.cpp
@@ -44,6 +44,13 @@ void UComponentElementEditorWorldInterface::NotifyMovementEnded(const FTypedElem
 	}
 }
 
+bool UComponentElementEditorWorldInterface::DeleteElements(TArrayView<const FTypedElementHandle> InElementHandles, UWorld* InWorld, UTypedElementSelectionSet* InSelectionSet, const FTypedElementDeletionOptions& InDeletionOptions)
+{
+	const TArray<UActorComponent*> ComponentsToDelete = ComponentElementDataUtil::GetComponentsFromHandles(InElementHandles);
+	return ComponentsToDelete.Num() > 0
+		&& GUnrealEd->DeleteComponents(ComponentsToDelete, InSelectionSet, InDeletionOptions.VerifyDeletionCanHappen());
+}
+
 void UComponentElementEditorWorldInterface::DuplicateElements(TArrayView<const FTypedElementHandle> InElementHandles, UWorld* InWorld, bool bOffsetLocations, TArray<FTypedElementHandle>& OutNewElements)
 {
 	const TArray<UActorComponent*> ComponentsToDuplicate = ComponentElementDataUtil::GetComponentsFromHandles(InElementHandles);
diff --git a/Engine/Source/Editor/UnrealEd/Private/Elements/Component/ComponentElementEditorWorldInterface.h b/Engine/Source/Editor/UnrealEd/Private/Elements/Component/ComponentElementEditorWorldInterface.h
--- a/Engine/Source/Editor/UnrealEd/Private/Elements/Component/ComponentElementEditorWorldInterface.h
+++ b/Engine/Source/Editor/UnrealEd/Private/Elements/Component/ComponentElementEditorWorldInterface.h
@@ -14,5 +14,6 @@ public:
 	virtual void NotifyMovementStarted(const FTypedElementHandle& InElementHandle) override;
 	virtual void NotifyMovementOngoing(const FTypedElementHandle& InElementHandle) override;
 	virtual void NotifyMovementEnded(const FTypedElementHandle& InElementHandle) override;
+	virtual bool DeleteElements(TArrayView<const FTypedElementHandle> InElementHandles, UWorld* InWorld, UTypedElementSelectionSet* InSelectionSet, const FTypedElementDeletionOptions& InDeletionOptions) override;
 	virtual void DuplicateElements(TArrayView<const FTypedElementHandle> InElementHandles, UWorld* InWorld, bool bOffsetLocations, TArray<FTypedElementHandle>& OutNewElements) override;
 };
diff --git a/Engine/Source/Editor/UnrealEd/Private/LevelEditorViewport.cpp b/Engine/Source/Editor/UnrealEd/Private/LevelEditorViewport.cpp
--- a/Engine/Source/Editor/UnrealEd/Private/LevelEditorViewport.cpp
+++ b/Engine/Source/Editor/UnrealEd/Private/LevelEditorViewport.cpp
@@ -40,6 +40,7 @@
 #include "EditorSupportDelegates.h"
 #include "Elements/Framework/TypedElementRegistry.h"
 #include "Elements/Framework/EngineElementsLibrary.h"
+#include "Elements/Framework/TypedElementCommonActions.h"
 #include "Elements/Framework/TypedElementListObjectUtil.h"
 #include "Elements/Framework/TypedElementViewportInteraction.h"
 #include "Elements/Actor/ActorElementLevelEditorViewportInteractionCustomization.h"
@@ -2573,10 +2574,12 @@ bool FLevelEditorViewportClient::InputWidgetDelta(FViewport* InViewport, EAxisLi
 						{
 							bDuplicateOnNextDrag = false;
 							
-							if (!ModeTools->ProcessEditDuplicate())
+							TSharedPtr<ILevelEditor> LevelEditor = ParentLevelEditor.Pin();
+							UTypedElementCommonActions* CommonActions = LevelEditor ? LevelEditor->GetCommonActions() : nullptr;
+							if (CommonActions)
 							{
 								CacheElementsToManipulate();
-								TArray<FTypedElementHandle> DuplicatedElements = UEngineElementsLibrary::DuplicateElements(ElementsToManipulate, GetWorld(), /*bOffsetLocations*/false);
+								TArray<FTypedElementHandle> DuplicatedElements = CommonActions->DuplicateElements(ElementsToManipulate, GetWorld(), /*bOffsetLocations*/false);
 								ElementsToManipulate->Reset();
 
 								// Exclusively select the new elements, so that future gizmo interaction manipulates those items instead
@@ -2589,9 +2592,9 @@ bool FLevelEditorViewportClient::InputWidgetDelta(FViewport* InViewport, EAxisLi
 
 									SelectionSet->SetSelection(DuplicatedElements, SelectionOptions);
 								}
-							}
 
-							RedrawAllViewportsIntoThisScene();
+								RedrawAllViewportsIntoThisScene();
+							}
 						}
 					}
 				}
diff --git a/Engine/Source/Editor/UnrealEd/Private/Toolkits/SStandaloneAssetEditorToolkitHost.cpp b/Engine/Source/Editor/UnrealEd/Private/Toolkits/SStandaloneAssetEditorToolkitHost.cpp
--- a/Engine/Source/Editor/UnrealEd/Private/Toolkits/SStandaloneAssetEditorToolkitHost.cpp
+++ b/Engine/Source/Editor/UnrealEd/Private/Toolkits/SStandaloneAssetEditorToolkitHost.cpp
@@ -344,6 +344,12 @@ UWorld* SStandaloneAssetEditorToolkitHost::GetWorld() const
 }
 
 
+UTypedElementCommonActions* SStandaloneAssetEditorToolkitHost::GetCommonActions() const
+{
+	return nullptr;
+}
+
+
 FReply SStandaloneAssetEditorToolkitHost::OnKeyDown( const FGeometry& MyGeometry, const FKeyEvent& InKeyEvent )
 {
 	// Check to see if any of the actions for the level editor can be processed by the current event
diff --git a/Engine/Source/Editor/UnrealEd/Private/Toolkits/SStandaloneAssetEditorToolkitHost.h b/Engine/Source/Editor/UnrealEd/Private/Toolkits/SStandaloneAssetEditorToolkitHost.h
--- a/Engine/Source/Editor/UnrealEd/Private/Toolkits/SStandaloneAssetEditorToolkitHost.h
+++ b/Engine/Source/Editor/UnrealEd/Private/Toolkits/SStandaloneAssetEditorToolkitHost.h
@@ -56,6 +56,7 @@ public:
 	virtual void OnToolkitHostingStarted( const TSharedRef< class IToolkit >& Toolkit ) override;
 	virtual void OnToolkitHostingFinished( const TSharedRef< class IToolkit >& Toolkit ) override;
 	virtual UWorld* GetWorld() const override;
+	virtual UTypedElementCommonActions* GetCommonActions() const override;
 	void CreateDefaultStandaloneMenuBar(UToolMenu* MenuBar);
 
 	/** SWidget overrides */
diff --git a/Engine/Source/Editor/UnrealEd/Private/UnrealEdSrv.cpp b/Engine/Source/Editor/UnrealEd/Private/UnrealEdSrv.cpp
--- a/Engine/Source/Editor/UnrealEd/Private/UnrealEdSrv.cpp
+++ b/Engine/Source/Editor/UnrealEd/Private/UnrealEdSrv.cpp
@@ -14,6 +14,7 @@
 #include "UObject/Class.h"
 #include "UObject/Package.h"
 #include "UObject/UnrealType.h"
+#include "UObject/StrongObjectPtr.h"
 #include "InputCoreTypes.h"
 #include "Input/Reply.h"
 #include "Widgets/DeclarativeSyntaxSupport.h"
@@ -83,6 +84,9 @@
 #include "MatineeExporter.h"
 #include "FbxExporter.h"
 #include "DesktopPlatformModule.h"
+#include "Elements/Framework/TypedElementList.h"
+#include "Elements/Framework/TypedElementRegistry.h"
+#include "Elements/Framework/TypedElementCommonActions.h"
 #include "SnappingUtils.h"
 #include "AssetSelection.h"
 #include "HighResScreenshot.h"
@@ -2486,18 +2490,16 @@ bool UUnrealEdEngine::Exec_Actor( UWorld* InWorld, const TCHAR* Str, FOutputDevi
 	}
 	else if( FParse::Command(&Str,TEXT("DELETE")) )		// ACTOR SELECT DELETE
 	{
-	
-		bool bHandled = false;
-		bHandled |= GLevelEditorModeTools().ProcessEditDelete();
-
-		// if not specially handled by the current editing mode,
-		if (!bHandled)
+		if (TSharedPtr<ILevelEditor> LevelEditor = GCurrentLevelEditingViewportClient->ParentLevelEditor.Pin())
 		{
-			UEditorActorSubsystem* EditorActorSubsystem = GEditor->GetEditorSubsystem<UEditorActorSubsystem>();
-
-			if (EditorActorSubsystem)
+			if (UTypedElementCommonActions* CommonActions = LevelEditor->GetCommonActions())
 			{
-				EditorActorSubsystem->DeleteSelectedActors(InWorld);
+				UTypedElementSelectionSet* SelectionSet = GetSelectedActors()->GetElementSelectionSet();
+
+				TStrongObjectPtr<UTypedElementList> ElementsToDelete(UTypedElementRegistry::GetInstance()->CreateElementList());
+				CommonActions->GetSelectedElementsToDelete(SelectionSet, ElementsToDelete.Get());
+				CommonActions->DeleteElements(ElementsToDelete.Get(), InWorld, SelectionSet, FTypedElementDeletionOptions());
+				ElementsToDelete->Reset();
 			}
 		}
 		return true;
diff --git a/Engine/Source/Runtime/Engine/Private/Elements/Framework/EngineElementsLibrary.cpp b/Engine/Source/Runtime/Engine/Private/Elements/Framework/EngineElementsLibrary.cpp
--- a/Engine/Source/Runtime/Engine/Private/Elements/Framework/EngineElementsLibrary.cpp
+++ b/Engine/Source/Runtime/Engine/Private/Elements/Framework/EngineElementsLibrary.cpp
@@ -216,48 +216,3 @@ FTypedElementHandle UEngineElementsLibrary::AcquireEditorComponentElementHandle(
 	return EngineElementsLibraryUtil::AcquireEditorTypedElementHandle<UActorComponent, FComponentElementData>(Component, GComponentElementOwnerStore, &UEngineElementsLibrary::CreateComponentElement, bAllowCreate);
 }
 #endif
-
-TArray<FTypedElementHandle> UEngineElementsLibrary::DuplicateElements(const TArray<FTypedElementHandle>& ElementHandles, UWorld* World, bool bOffsetLocations)
-{
-	return DuplicateElements(MakeArrayView(ElementHandles), World, bOffsetLocations);
-}
-
-TArray<FTypedElementHandle> UEngineElementsLibrary::DuplicateElements(TArrayView<const FTypedElementHandle> ElementHandles, UWorld* World, bool bOffsetLocations)
-{
-	TMap<FTypedHandleTypeId, TArray<FTypedElementHandle>> ElementsToDuplicateByType;
-	TypedElementUtil::BatchElementsByType(ElementHandles, ElementsToDuplicateByType);
-
-	TArray<FTypedElementHandle> NewElements;
-	NewElements.Reserve(ElementHandles.Num());
-
-	UTypedElementRegistry* Registry = UTypedElementRegistry::GetInstance();
-	for (const auto& ElementsByTypePair : ElementsToDuplicateByType)
-	{
-		if (UTypedElementWorldInterface* WorldInterface = Registry->GetElementInterface<UTypedElementWorldInterface>(ElementsByTypePair.Key))
-		{
-			WorldInterface->DuplicateElements(ElementsByTypePair.Value, World, bOffsetLocations, NewElements);
-		}
-	}
-
-	return NewElements;
-}
-
-TArray<FTypedElementHandle> UEngineElementsLibrary::DuplicateElements(const UTypedElementList* ElementList, UWorld* World, bool bOffsetLocations)
-{
-	TMap<FTypedHandleTypeId, TArray<FTypedElementHandle>> ElementsToDuplicateByType;
-	TypedElementUtil::BatchElementsByType(ElementList, ElementsToDuplicateByType);
-
-	TArray<FTypedElementHandle> NewElements;
-	NewElements.Reserve(ElementList->Num());
-
-	UTypedElementRegistry* Registry = UTypedElementRegistry::GetInstance();
-	for (const auto& ElementsByTypePair : ElementsToDuplicateByType)
-	{
-		if (UTypedElementWorldInterface* WorldInterface = Registry->GetElementInterface<UTypedElementWorldInterface>(ElementsByTypePair.Key))
-		{
-			WorldInterface->DuplicateElements(ElementsByTypePair.Value, World, bOffsetLocations, NewElements);
-		}
-	}
-
-	return NewElements;
-}
diff --git a/Engine/Source/Runtime/Engine/Public/Elements/Framework/EngineElementsLibrary.h b/Engine/Source/Runtime/Engine/Public/Elements/Framework/EngineElementsLibrary.h
--- a/Engine/Source/Runtime/Engine/Public/Elements/Framework/EngineElementsLibrary.h
+++ b/Engine/Source/Runtime/Engine/Public/Elements/Framework/EngineElementsLibrary.h
@@ -53,11 +53,6 @@ public:
 	static FTypedElementHandle AcquireEditorComponentElementHandle(const UActorComponent* Component, const bool bAllowCreate = true);
 #endif
 
-	UFUNCTION(BlueprintCallable, Category="TypedElementFramework|Util")
-	static TArray<FTypedElementHandle> DuplicateElements(const TArray<FTypedElementHandle>& ElementHandles, UWorld* World, bool bOffsetLocations);
-	static TArray<FTypedElementHandle> DuplicateElements(TArrayView<const FTypedElementHandle> ElementHandles, UWorld* World, bool bOffsetLocations);
-	static TArray<FTypedElementHandle> DuplicateElements(const UTypedElementList* ElementList, UWorld* World, bool bOffsetLocations);
-
 private:
 #if WITH_EDITOR
 	static void DestroyUnreachableEditorObjectElements();
diff --git a/Engine/Source/Runtime/Engine/Public/Elements/Interfaces/TypedElementWorldInterface.h b/Engine/Source/Runtime/Engine/Public/Elements/Interfaces/TypedElementWorldInterface.h
--- a/Engine/Source/Runtime/Engine/Public/Elements/Interfaces/TypedElementWorldInterface.h
+++ b/Engine/Source/Runtime/Engine/Public/Elements/Interfaces/TypedElementWorldInterface.h
@@ -7,8 +7,35 @@
 
 class ULevel;
 class UWorld;
+class UTypedElementSelectionSet;
 struct FCollisionShape;
 
+USTRUCT(BlueprintType)
+struct FTypedElementDeletionOptions
+{
+	GENERATED_BODY()
+
+public:
+	FTypedElementDeletionOptions& SetVerifyDeletionCanHappen(const bool InVerifyDeletionCanHappen) { bVerifyDeletionCanHappen = InVerifyDeletionCanHappen; return *this; }
+	bool VerifyDeletionCanHappen() const { return bVerifyDeletionCanHappen; }
+
+	FTypedElementDeletionOptions& SetWarnAboutReferences(const bool InWarnAboutReferences) { bWarnAboutReferences = InWarnAboutReferences; return *this; }
+	bool WarnAboutReferences() const { return bWarnAboutReferences; }
+
+	FTypedElementDeletionOptions& SetWarnAboutSoftReferences(const bool InWarnAboutSoftReferences) { bWarnAboutSoftReferences = InWarnAboutSoftReferences; return *this; }
+	bool WarnAboutSoftReferences() const { return bWarnAboutSoftReferences; }
+	
+private:
+	UPROPERTY(BlueprintReadWrite, Category="TypedElementInterfaces|World|DeletionOptions", meta=(AllowPrivateAccess=true))
+	bool bVerifyDeletionCanHappen = false;
+
+	UPROPERTY(BlueprintReadWrite, Category="TypedElementInterfaces|World|DeletionOptions", meta=(AllowPrivateAccess=true))
+	bool bWarnAboutReferences = true;
+
+	UPROPERTY(BlueprintReadWrite, Category="TypedElementInterfaces|World|DeletionOptions", meta=(AllowPrivateAccess=true))
+	bool bWarnAboutSoftReferences = true;
+};
+
 UCLASS(Abstract)
 class ENGINE_API UTypedElementWorldInterface : public UTypedElementInterface
 {
@@ -116,6 +143,28 @@ public:
 		return false;
 	}
 
+	/**
+	 * Can the given element be deleted?
+	 */
+	UFUNCTION(BlueprintCallable, Category="TypedElementInterfaces|World")
+	virtual bool CanDeleteElement(const FTypedElementHandle& InElementHandle) { return false; }
+
+	/**
+	 * Delete the given element.
+	 * @note Default version calls DeleteElements with a single element.
+	 */
+	UFUNCTION(BlueprintCallable, Category="TypedElementInterfaces|World")
+	virtual bool DeleteElement(const FTypedElementHandle& InElementHandle, UWorld* InWorld, UTypedElementSelectionSet* InSelectionSet, const FTypedElementDeletionOptions& InDeletionOptions)
+	{
+		return DeleteElements(MakeArrayView(&InElementHandle, 1), InWorld, InSelectionSet, InDeletionOptions);
+	}
+
+	/**
+	 * Delete the given set of elements.
+	 * @note If you want to delete an array of elements that are potentially different types, you probably want to use the higher-level UTypedElementCommonActions::DeleteElements function instead.
+	 */
+	virtual bool DeleteElements(TArrayView<const FTypedElementHandle> InElementHandles, UWorld* InWorld, UTypedElementSelectionSet* InSelectionSet, const FTypedElementDeletionOptions& InDeletionOptions) { return false; }
+
 	/**
 	 * Duplicate the given element.
 	 * @note Default version calls DuplicateElements with a single element.
@@ -130,7 +179,7 @@ public:
 
 	/**
 	 * Duplicate the given set of elements.
-	 * @note If you want to duplicate an array of elements that are potentially different types, you probably want to use the higher-level UEngineElementsLibrary::DuplicateElements function instead.
+	 * @note If you want to duplicate an array of elements that are potentially different types, you probably want to use the higher-level UTypedElementCommonActions::DuplicateElements function instead.
 	 */
 	virtual void DuplicateElements(TArrayView<const FTypedElementHandle> InElementHandles, UWorld* InWorld, bool bOffsetLocations, TArray<FTypedElementHandle>& OutNewElements)
 	{
@@ -156,5 +205,7 @@ struct TTypedElement<UTypedElementWorldInterface> : public TTypedElementBase<UTy
 	void NotifyMovementEnded() const { InterfacePtr->NotifyMovementEnded(*this); }
 	bool FindSuitableTransformAtPoint(const FTransform& InPotentialTransform, FTransform& OutSuitableTransform) const { return InterfacePtr->FindSuitableTransformAtPoint(*this, InPotentialTransform, OutSuitableTransform); }
 	bool FindSuitableTransformAlongPath(const FVector& InPathStart, const FVector& InPathEnd, const FCollisionShape& InTestShape, TArrayView<const FTypedElementHandle> InElementsToIgnore, FTransform& OutSuitableTransform) const { return InterfacePtr->FindSuitableTransformAlongPath(*this, InPathStart, InPathEnd, InTestShape, InElementsToIgnore, OutSuitableTransform); }
+	bool CanDeleteElement() const { return InterfacePtr->CanDeleteElement(*this); }
+	bool DeleteElement(UWorld* InWorld, UTypedElementSelectionSet* InSelectionSet, const FTypedElementDeletionOptions& InDeletionOptions) const { return InterfacePtr->DeleteElement(*this, InWorld, InSelectionSet, InDeletionOptions); }
 	FTypedElementHandle DuplicateElement(UWorld* InWorld, bool bOffsetLocations) const { return InterfacePtr->DuplicateElement(*this, InWorld, bOffsetLocations); }
 };
