Add virtual destructors to IToolkit and IToolkitHost.

diff --git a/Engine/Source/Editor/EditorFramework/Public/Toolkits/IToolkit.h b/Engine/Source/Editor/EditorFramework/Public/Toolkits/IToolkit.h
--- a/Engine/Source/Editor/EditorFramework/Public/Toolkits/IToolkit.h
+++ b/Engine/Source/Editor/EditorFramework/Public/Toolkits/IToolkit.h
@@ -61,6 +61,8 @@ class IToolkit
 {
 public:
 
+	virtual ~IToolkit() = default;
+
 	/** Register tabs that this toolkit can spawn with the TabManager */
 	virtual void RegisterTabSpawners(const TSharedRef<FTabManager>& TabManager) = 0;
 
diff --git a/Engine/Source/Editor/EditorFramework/Public/Toolkits/IToolkitHost.h b/Engine/Source/Editor/EditorFramework/Public/Toolkits/IToolkitHost.h
--- a/Engine/Source/Editor/EditorFramework/Public/Toolkits/IToolkitHost.h
+++ b/Engine/Source/Editor/EditorFramework/Public/Toolkits/IToolkitHost.h
@@ -15,6 +15,7 @@ class IToolkitHost
 {
 
 public:
+	virtual ~IToolkitHost() = default;
 
 	/* Notificaton when the active viewport changed */
 	DECLARE_MULTICAST_DELEGATE_TwoParams(FOnActiveViewportChanged,  TSharedPtr<IAssetViewport>, TSharedPtr<IAssetViewport>)
