Fixed incorrect comment

#rnx

diff --git a/Engine/Source/Editor/EditorFramework/Public/Elements/Framework/TypedElementViewportInteraction.h b/Engine/Source/Editor/EditorFramework/Public/Elements/Framework/TypedElementViewportInteraction.h
--- a/Engine/Source/Editor/EditorFramework/Public/Elements/Framework/TypedElementViewportInteraction.h
+++ b/Engine/Source/Editor/EditorFramework/Public/Elements/Framework/TypedElementViewportInteraction.h
@@ -91,17 +91,20 @@ class EDITORFRAMEWORK_API UTypedElementViewportInteraction : public UObject, pub
 
 public:
 	/**
-	 * Notify that the gizmo is potentially about to start manipulating the transform of the given set of elements (calculated from calling GetSelectedElementsToMove).
+	 * Notify that the gizmo is potentially about to start manipulating the transform of the given set of elements.
+	 * @note This list should have been pre-normalized via UTypedElementSelectionSet::GetNormalizedSelection or UTypedElementSelectionSet::GetNormalizedElementList, or FLevelEditorViewportClient::GetElementsToManipulate.
 	 */
 	void BeginGizmoManipulation(FTypedElementListConstRef InElementsToMove, const UE::Widget::EWidgetMode InWidgetMode);
 
 	/**
-	 * Notify that the gizmo has manipulated the transform of the given set of elements (calculated from calling GetSelectedElementsToMove) by the given delta.
+	 * Notify that the gizmo has manipulated the transform of the given set of elements.
+	 * @note This list should have been pre-normalized via UTypedElementSelectionSet::GetNormalizedSelection or UTypedElementSelectionSet::GetNormalizedElementList, or FLevelEditorViewportClient::GetElementsToManipulate.
 	 */
 	void UpdateGizmoManipulation(FTypedElementListConstRef InElementsToMove, const UE::Widget::EWidgetMode InWidgetMode, const EAxisList::Type InDragAxis, const FInputDeviceState& InInputState, const FTransform& InDeltaTransform);
 	
 	/**
-	 * Notify that the gizmo has finished manipulating the transform of the given set of elements (calculated from calling GetSelectedElementsToMove).
+	 * Notify that the gizmo has finished manipulating the transform of the given set of elements.
+	 * @note This list should have been pre-normalized via UTypedElementSelectionSet::GetNormalizedSelection or UTypedElementSelectionSet::GetNormalizedElementList, or FLevelEditorViewportClient::GetElementsToManipulate.
 	 */
 	void EndGizmoManipulation(FTypedElementListConstRef InElementsToMove, const UE::Widget::EWidgetMode InWidgetMode, const ETypedElementViewportInteractionGizmoManipulationType InManipulationType);
 
