Fixing the SMInstanceElements rotating differently to an offseted pivot then a normal actor or component.

The issue was that weren't using the same pivot  and the same quaternion math for the rotation.

diff --git a/Engine/Source/Editor/EditorFramework/Private/Elements/Framework/TypedElementViewportInteraction.cpp b/Engine/Source/Editor/EditorFramework/Private/Elements/Framework/TypedElementViewportInteraction.cpp
--- a/Engine/Source/Editor/EditorFramework/Private/Elements/Framework/TypedElementViewportInteraction.cpp
+++ b/Engine/Source/Editor/EditorFramework/Private/Elements/Framework/TypedElementViewportInteraction.cpp
@@ -1,12 +1,21 @@
 // Copyright Epic Games, Inc. All Rights Reserved.
 
 #include "Elements/Framework/TypedElementViewportInteraction.h"
-#include "Elements/Framework/TypedElementSelectionSet.h"
+
+#include "EditorModeManager.h"
 #include "Elements/Framework/TypedElementRegistry.h"
+#include "Elements/Framework/TypedElementSelectionSet.h"
 #include "Elements/Framework/TypedElementUtil.h"
+#include "Toolkits/IToolkitHost.h"
 
 bool FTypedElementViewportInteractionCustomization::GetGizmoPivotLocation(const TTypedElement<ITypedElementWorldInterface>& InElementWorldHandle, const UE::Widget::EWidgetMode InWidgetMode, FVector& OutPivotLocation)
 {
+	if (const IToolkitHost* ToolkitHostPtr = GetToolkitHost())
+	{
+		OutPivotLocation = ToolkitHostPtr->GetEditorModeManager().PivotLocation;
+		return true;
+	}
+
 	FTransform ElementWorldTransform;
 	if (InElementWorldHandle.GetWorldTransform(ElementWorldTransform))
 	{
@@ -36,7 +45,9 @@ void FTypedElementViewportInteractionCustomization::GizmoManipulationDeltaUpdate
 			const FQuat DeltaRotation = InDeltaTransform.GetRotation();
 			if (!DeltaRotation.Rotator().IsZero())
 			{
-				ElementWorldTransform.SetRotation(ElementWorldTransform.GetRotation() * DeltaRotation);
+				const FQuat ActorQ =  ElementWorldTransform.GetRotation();
+				const FQuat ResultQ = DeltaRotation * ActorQ;
+				ElementWorldTransform.SetRotation(ResultQ.GetNormalized());
 
 				FVector ElementLocation = ElementWorldTransform.GetTranslation();
 				ElementLocation -= InPivotLocation;
diff --git a/Engine/Source/Editor/EditorFramework/Public/Elements/Framework/TypedElementViewportInteraction.h b/Engine/Source/Editor/EditorFramework/Public/Elements/Framework/TypedElementViewportInteraction.h
--- a/Engine/Source/Editor/EditorFramework/Public/Elements/Framework/TypedElementViewportInteraction.h
+++ b/Engine/Source/Editor/EditorFramework/Public/Elements/Framework/TypedElementViewportInteraction.h
@@ -7,6 +7,7 @@
 #include "Elements/Framework/TypedElementListFwd.h"
 #include "Elements/Framework/TypedElementInterfaceCustomization.h"
 #include "Elements/Interfaces/TypedElementWorldInterface.h"
+#include "Elements/Framework/TypedElementAssetEditorToolkitHostMixin.h"
 #include "TypedElementViewportInteraction.generated.h"
 
 class UTypedElementSelectionSet;
@@ -22,7 +23,7 @@ enum class ETypedElementViewportInteractionGizmoManipulationType : uint8
 /**
  * Customization used to allow asset editors (such as the level editor) to override the base behavior of viewport interaction.
  */
-class EDITORFRAMEWORK_API FTypedElementViewportInteractionCustomization
+class EDITORFRAMEWORK_API FTypedElementViewportInteractionCustomization : public FTypedElementAssetEditorToolkitHostMixin
 {
 public:
 	virtual ~FTypedElementViewportInteractionCustomization() = default;
diff --git a/Engine/Source/Editor/UnrealEd/Private/Elements/Actor/ActorElementEditorViewportInteractionCustomization.cpp b/Engine/Source/Editor/UnrealEd/Private/Elements/Actor/ActorElementEditorViewportInteractionCustomization.cpp
--- a/Engine/Source/Editor/UnrealEd/Private/Elements/Actor/ActorElementEditorViewportInteractionCustomization.cpp
+++ b/Engine/Source/Editor/UnrealEd/Private/Elements/Actor/ActorElementEditorViewportInteractionCustomization.cpp
@@ -5,8 +5,6 @@
 #include "GameFramework/Actor.h"
 
 #include "Editor.h"
-#include "EditorModeManager.h"
-#include "Toolkits/IToolkitHost.h"
 #include "AI/NavigationSystemBase.h"
 #include "Components/BrushComponent.h"
 #include "Engine/DocumentationActor.h"
@@ -14,17 +12,6 @@
 #include "Elements/Interfaces/TypedElementHierarchyInterface.h"
 #include "Elements/Framework/TypedElementRegistry.h"
 
-bool FActorElementEditorViewportInteractionCustomization::GetGizmoPivotLocation(const TTypedElement<ITypedElementWorldInterface>& InElementWorldHandle, const UE::Widget::EWidgetMode InWidgetMode, FVector& OutPivotLocation)
-{
-	if (const IToolkitHost* ToolkitHostPtr = GetToolkitHost())
-	{
-		OutPivotLocation = ToolkitHostPtr->GetEditorModeManager().PivotLocation;
-		return true;
-	}
-	
-	return FTypedElementViewportInteractionCustomization::GetGizmoPivotLocation(InElementWorldHandle, InWidgetMode, OutPivotLocation);
-}
-
 void FActorElementEditorViewportInteractionCustomization::GizmoManipulationDeltaUpdate(const TTypedElement<ITypedElementWorldInterface>& InElementWorldHandle, const UE::Widget::EWidgetMode InWidgetMode, const EAxisList::Type InDragAxis, const FInputDeviceState& InInputState, const FTransform& InDeltaTransform, const FVector& InPivotLocation)
 {
 	AActor* Actor = ActorElementDataUtil::GetActorFromHandleChecked(InElementWorldHandle);
diff --git a/Engine/Source/Editor/UnrealEd/Private/Elements/Component/ComponentElementEditorViewportInteractionCustomization.cpp b/Engine/Source/Editor/UnrealEd/Private/Elements/Component/ComponentElementEditorViewportInteractionCustomization.cpp
--- a/Engine/Source/Editor/UnrealEd/Private/Elements/Component/ComponentElementEditorViewportInteractionCustomization.cpp
+++ b/Engine/Source/Editor/UnrealEd/Private/Elements/Component/ComponentElementEditorViewportInteractionCustomization.cpp
@@ -5,8 +5,6 @@
 #include "Components/SceneComponent.h"
 
 #include "Editor.h"
-#include "EditorModeManager.h"
-#include "Toolkits/IToolkitHost.h"
 #include "EditorSupportDelegates.h"
 
 bool FComponentElementEditorViewportInteractionCustomization::GetGizmoPivotLocation(const TTypedElement<ITypedElementWorldInterface>& InElementWorldHandle, const UE::Widget::EWidgetMode InWidgetMode, FVector& OutPivotLocation)
@@ -16,14 +14,7 @@ bool FComponentElementEditorViewportInteractionCustomization::GetGizmoPivotLocat
 	if (const USceneComponent* SceneComponent = Cast<USceneComponent>(Component))
 	{
 		FVector BasePivotLocation = FVector::ZeroVector;
-		if (const IToolkitHost* ToolkitHostPtr = GetToolkitHost())
-		{
-			BasePivotLocation = ToolkitHostPtr->GetEditorModeManager().PivotLocation;
-		}
-		else
-		{
-			FTypedElementViewportInteractionCustomization::GetGizmoPivotLocation(InElementWorldHandle, InWidgetMode, BasePivotLocation);
-		}
+		FTypedElementViewportInteractionCustomization::GetGizmoPivotLocation(InElementWorldHandle, InWidgetMode, BasePivotLocation);
 
 		// If necessary, transform the editor pivot location to be relative to the component's parent
 		const bool bIsRootComponent = SceneComponent->GetOwner()->GetRootComponent() == SceneComponent;
diff --git a/Engine/Source/Editor/UnrealEd/Private/LevelEditorViewport.cpp b/Engine/Source/Editor/UnrealEd/Private/LevelEditorViewport.cpp
--- a/Engine/Source/Editor/UnrealEd/Private/LevelEditorViewport.cpp
+++ b/Engine/Source/Editor/UnrealEd/Private/LevelEditorViewport.cpp
@@ -1829,6 +1829,10 @@ void FLevelEditorViewportClient::InitializeViewportInteraction()
 		Registry->OnProcessingDeferredElementsToDestroy().AddRaw(this, &FLevelEditorViewportClient::ResetElementsToManipulateFromProcessingDeferredElementsToDestroy);
 	}
 
+	{
+		ViewportInteraction->GetDefaultInterfaceCustomization()->SetToolkitHost(ParentLevelEditor.Pin().Get());
+	}
+
 	{
 		TUniquePtr<FActorElementLevelEditorViewportInteractionCustomization> ActorCustomization = MakeUnique<FActorElementLevelEditorViewportInteractionCustomization>();
 		ActorCustomization->SetLevelEditorViewportClient(this);
diff --git a/Engine/Source/Editor/UnrealEd/Public/Elements/Actor/ActorElementEditorViewportInteractionCustomization.h b/Engine/Source/Editor/UnrealEd/Public/Elements/Actor/ActorElementEditorViewportInteractionCustomization.h
--- a/Engine/Source/Editor/UnrealEd/Public/Elements/Actor/ActorElementEditorViewportInteractionCustomization.h
+++ b/Engine/Source/Editor/UnrealEd/Public/Elements/Actor/ActorElementEditorViewportInteractionCustomization.h
@@ -7,10 +7,9 @@
 
 class AActor;
 
-class UNREALED_API FActorElementEditorViewportInteractionCustomization : public FTypedElementViewportInteractionCustomization, public FTypedElementAssetEditorToolkitHostMixin
+class UNREALED_API FActorElementEditorViewportInteractionCustomization : public FTypedElementViewportInteractionCustomization
 {
 public:
-	virtual bool GetGizmoPivotLocation(const TTypedElement<ITypedElementWorldInterface>& InElementWorldHandle, const UE::Widget::EWidgetMode InWidgetMode, FVector& OutPivotLocation) override;
 	virtual void GizmoManipulationDeltaUpdate(const TTypedElement<ITypedElementWorldInterface>& InElementWorldHandle, const UE::Widget::EWidgetMode InWidgetMode, const EAxisList::Type InDragAxis, const FInputDeviceState& InInputState, const FTransform& InDeltaTransform, const FVector& InPivotLocation) override;
 	virtual void MirrorElement(const TTypedElement<ITypedElementWorldInterface>& InElementWorldHandle, const FVector& InMirrorScale, const FVector& InPivotLocation) override;
 	virtual bool GetFocusBounds(const TTypedElement<ITypedElementWorldInterface>& InElementWorldHandle, FBoxSphereBounds& OutBounds) override;
diff --git a/Engine/Source/Editor/UnrealEd/Public/Elements/Component/ComponentElementEditorViewportInteractionCustomization.h b/Engine/Source/Editor/UnrealEd/Public/Elements/Component/ComponentElementEditorViewportInteractionCustomization.h
--- a/Engine/Source/Editor/UnrealEd/Public/Elements/Component/ComponentElementEditorViewportInteractionCustomization.h
+++ b/Engine/Source/Editor/UnrealEd/Public/Elements/Component/ComponentElementEditorViewportInteractionCustomization.h
@@ -8,7 +8,7 @@
 class FProperty;
 class USceneComponent;
 
-class UNREALED_API FComponentElementEditorViewportInteractionCustomization : public FTypedElementViewportInteractionCustomization, public FTypedElementAssetEditorToolkitHostMixin
+class UNREALED_API FComponentElementEditorViewportInteractionCustomization : public FTypedElementViewportInteractionCustomization
 {
 public:
 	virtual bool GetGizmoPivotLocation(const TTypedElement<ITypedElementWorldInterface>& InElementWorldHandle, const UE::Widget::EWidgetMode InWidgetMode, FVector& OutPivotLocation) override;
diff --git a/Engine/Source/Runtime/TypedElementRuntime/Public/Elements/Framework/TypedElementInterfaceCustomization.h b/Engine/Source/Runtime/TypedElementRuntime/Public/Elements/Framework/TypedElementInterfaceCustomization.h
--- a/Engine/Source/Runtime/TypedElementRuntime/Public/Elements/Framework/TypedElementInterfaceCustomization.h
+++ b/Engine/Source/Runtime/TypedElementRuntime/Public/Elements/Framework/TypedElementInterfaceCustomization.h
@@ -58,6 +58,14 @@ public:
 		DefaultInterfaceCustomization = MoveTemp(InInterfaceCustomization);
 	}
 
+	/**
+	 * Return the default interface customization instance.
+	 */
+	CustomizationBaseType* GetDefaultInterfaceCustomization() const
+	{
+		return DefaultInterfaceCustomization.Get();
+	}
+
 	/**
 	 * Register an interface customization for the given element type.
 	 */
