Fixes for non-unity builds.

diff --git a/Engine/Source/Editor/EditorFramework/Public/Subsystems/PlacementSubsystem.h b/Engine/Source/Editor/EditorFramework/Public/Subsystems/PlacementSubsystem.h
--- a/Engine/Source/Editor/EditorFramework/Public/Subsystems/PlacementSubsystem.h
+++ b/Engine/Source/Editor/EditorFramework/Public/Subsystems/PlacementSubsystem.h
@@ -4,8 +4,8 @@
 
 #include "AssetRegistry/AssetData.h"
 #include "Elements/Framework/TypedElementHandle.h"
-#include "UObject/Interface.h"
 #include "EditorSubsystem.h"
+#include "Templates/SubclassOf.h"
 
 #include "PlacementSubsystem.generated.h"
 
diff --git a/Engine/Source/Editor/Persona/Public/IAssetFamily.h b/Engine/Source/Editor/Persona/Public/IAssetFamily.h
--- a/Engine/Source/Editor/Persona/Public/IAssetFamily.h
+++ b/Engine/Source/Editor/Persona/Public/IAssetFamily.h
@@ -4,6 +4,7 @@
 
 #include "CoreMinimal.h"
 #include "AssetData.h"
+#include "Styling/SlateColor.h"
 
 struct FSlateBrush;
 
diff --git a/Engine/Source/Runtime/Engine/Private/Animation/MeshDeformerProvider.cpp b/Engine/Source/Runtime/Engine/Private/Animation/MeshDeformerProvider.cpp
--- a/Engine/Source/Runtime/Engine/Private/Animation/MeshDeformerProvider.cpp
+++ b/Engine/Source/Runtime/Engine/Private/Animation/MeshDeformerProvider.cpp
@@ -1,6 +1,7 @@
 // Copyright Epic Games, Inc. All Rights Reserved.
 
 #include "Animation/MeshDeformerProvider.h"
+#include "Features/IModularFeatures.h"
 
 const FName IMeshDeformerProvider::ModularFeatureName(TEXT("MeshDeformer"));
 
