Further optimization for moving actor with many static mesh components.  Only invalidate the static lighting cache at the start of a move operation.  We don't invalidate the cache in "GizmoManipulationStarted", because the user can click on a handle (calling the start function) without actually moving, which should have no effect.  So a flag is used to detect the first time "GizmoManipulationDeltaUpdate" is called, and subsequent calls can skip repeatedly invalidating the lighting cache.

With this fix, on my machine, the cost of moving a blueprint with 1000 items drops from 45 ms to 7 ms.  Baseline tick is 8 ms (120 fps) when idling, so the results while moving improve from 19 fps to 65 fps.  The remaining time is mostly split between physics and navigation update.  Invalidating the static lighting cache is so costly because it both invalidates the render state of every item, and notifies packages of the change.

#rnx

diff --git a/Engine/Source/Editor/EditorFramework/Private/Elements/Framework/TypedElementViewportInteraction.cpp b/Engine/Source/Editor/EditorFramework/Private/Elements/Framework/TypedElementViewportInteraction.cpp
--- a/Engine/Source/Editor/EditorFramework/Private/Elements/Framework/TypedElementViewportInteraction.cpp
+++ b/Engine/Source/Editor/EditorFramework/Private/Elements/Framework/TypedElementViewportInteraction.cpp
@@ -54,7 +54,7 @@ void FTypedElementViewportInteractionCustomization::GizmoManipulationStarted(con
 	InElementWorldHandle.NotifyMovementStarted();
 }
 
-void FTypedElementViewportInteractionCustomization::GizmoManipulationDeltaUpdate(const TTypedElement<ITypedElementWorldInterface>& InElementWorldHandle, const UE::Widget::EWidgetMode InWidgetMode, const EAxisList::Type InDragAxis, const FInputDeviceState& InInputState, const FTransform& InDeltaTransform, const FVector& InPivotLocation)
+void FTypedElementViewportInteractionCustomization::GizmoManipulationDeltaUpdate(const TTypedElement<ITypedElementWorldInterface>& InElementWorldHandle, const UE::Widget::EWidgetMode InWidgetMode, const EAxisList::Type InDragAxis, const FInputDeviceState& InInputState, const FTransform& InDeltaTransform, const FVector& InPivotLocation, ETypedElementViewportInteractionDragMovementType GizmoDeltaType)
 {
 	FTransform ElementWorldTransform;
 	if (InElementWorldHandle.GetWorldTransform(ElementWorldTransform))
@@ -191,16 +191,16 @@ void UTypedElementViewportInteraction::BeginGizmoManipulation(FTypedElementListC
 	});
 }
 
-void UTypedElementViewportInteraction::UpdateGizmoManipulation(FTypedElementListConstRef InElementsToMove, const UE::Widget::EWidgetMode InWidgetMode, const EAxisList::Type InDragAxis, const FInputDeviceState& InInputState, const FTransform& InDeltaTransform)
+void UTypedElementViewportInteraction::UpdateGizmoManipulation(FTypedElementListConstRef InElementsToMove, const UE::Widget::EWidgetMode InWidgetMode, const EAxisList::Type InDragAxis, const FInputDeviceState& InInputState, const FTransform& InDeltaTransform, ETypedElementViewportInteractionDragMovementType GizmoDeltaType)
 {
-	InElementsToMove->ForEachElementHandle([this, InWidgetMode, InDragAxis, &InInputState, &InDeltaTransform](const FTypedElementHandle& InElementToMove)
+	InElementsToMove->ForEachElementHandle([this, InWidgetMode, InDragAxis, &InInputState, &InDeltaTransform, GizmoDeltaType](const FTypedElementHandle& InElementToMove)
 	{
 		FTypedElementViewportInteractionElement ViewportInteractionElement = ResolveViewportInteractionElement(InElementToMove);
 		if (ViewportInteractionElement)
 		{
 			FVector PivotLocation = FVector::ZeroVector;
 			ViewportInteractionElement.GetGizmoPivotLocation(InWidgetMode, PivotLocation);
-			ViewportInteractionElement.GizmoManipulationDeltaUpdate(InWidgetMode, InDragAxis, InInputState, InDeltaTransform, PivotLocation);
+			ViewportInteractionElement.GizmoManipulationDeltaUpdate(InWidgetMode, InDragAxis, InInputState, InDeltaTransform, PivotLocation, GizmoDeltaType);
 		}
 		return true;
 	});
@@ -238,7 +238,7 @@ void UTypedElementViewportInteraction::ApplyDeltaToElement(const FTypedElementHa
 	{
 		FVector PivotLocation = FVector::ZeroVector;
 		ViewportInteractionElement.GetGizmoPivotLocation(InWidgetMode, PivotLocation);
-		ViewportInteractionElement.GizmoManipulationDeltaUpdate(InWidgetMode, InDragAxis, InInputState, InDeltaTransform, PivotLocation);
+		ViewportInteractionElement.GizmoManipulationDeltaUpdate(InWidgetMode, InDragAxis, InInputState, InDeltaTransform, PivotLocation, ETypedElementViewportInteractionDragMovementType::None);
 	}
 }
 
diff --git a/Engine/Source/Editor/EditorFramework/Public/Elements/Framework/TypedElementViewportInteraction.h b/Engine/Source/Editor/EditorFramework/Public/Elements/Framework/TypedElementViewportInteraction.h
--- a/Engine/Source/Editor/EditorFramework/Public/Elements/Framework/TypedElementViewportInteraction.h
+++ b/Engine/Source/Editor/EditorFramework/Public/Elements/Framework/TypedElementViewportInteraction.h
@@ -31,6 +31,20 @@ enum class ETypedElementViewportInteractionGizmoManipulationType : uint8
 	Click,
 };
 
+/**
+ * Setting to identify if a movement is a drag operation, and if so, whether it is a first or subsequent move.  Some clients may apply
+ * optimizations where certain logic (such as dirtying state) only runs on the first movement delta.
+ */
+enum class ETypedElementViewportInteractionDragMovementType : uint8
+{
+	/** Not part of a drag operation */
+	None,
+	/** First delta movement of a drag operation */
+	FirstMove,
+	/** Subsequent delta movement of a drag operation */
+	SubsequentMove
+};
+
 /**
  * Customization used to allow asset editors (such as the level editor) to override the base behavior of viewport interaction.
  */
@@ -43,7 +57,7 @@ public:
 	virtual bool GetGizmoPivotLocation(const TTypedElement<ITypedElementWorldInterface>& InElementWorldHandle, const UE::Widget::EWidgetMode InWidgetMode, FVector& OutPivotLocation);
 	virtual void PreGizmoManipulationStarted(TArrayView<const FTypedElementHandle> InElementHandles, const UE::Widget::EWidgetMode InWidgetMode);
 	virtual void GizmoManipulationStarted(const TTypedElement<ITypedElementWorldInterface>& InElementWorldHandle, const UE::Widget::EWidgetMode InWidgetMode);
-	virtual void GizmoManipulationDeltaUpdate(const TTypedElement<ITypedElementWorldInterface>& InElementWorldHandle, const UE::Widget::EWidgetMode InWidgetMode, const EAxisList::Type InDragAxis, const FInputDeviceState& InInputState, const FTransform& InDeltaTransform, const FVector& InPivotLocation);
+	virtual void GizmoManipulationDeltaUpdate(const TTypedElement<ITypedElementWorldInterface>& InElementWorldHandle, const UE::Widget::EWidgetMode InWidgetMode, const EAxisList::Type InDragAxis, const FInputDeviceState& InInputState, const FTransform& InDeltaTransform, const FVector& InPivotLocation, ETypedElementViewportInteractionDragMovementType GizmoDeltaType);
 	virtual void GizmoManipulationStopped(const TTypedElement<ITypedElementWorldInterface>& InElementWorldHandle, const UE::Widget::EWidgetMode InWidgetMode, const ETypedElementViewportInteractionGizmoManipulationType InManipulationType);
 	virtual void PostGizmoManipulationStopped(TArrayView<const FTypedElementHandle> InElementHandles, const UE::Widget::EWidgetMode InWidgetMode);
 	virtual void MirrorElement(const TTypedElement<ITypedElementWorldInterface>& InElementWorldHandle, const FVector& InMirrorScale, const FVector& InPivotLocation);
@@ -84,7 +98,7 @@ public:
 	//~ See UTypedElementViewportInteraction for API docs
 	bool GetGizmoPivotLocation(const UE::Widget::EWidgetMode InWidgetMode, FVector& OutPivotLocation) const { return ViewportInteractionCustomization->GetGizmoPivotLocation(ElementWorldHandle, InWidgetMode, OutPivotLocation); }
 	void GizmoManipulationStarted(const UE::Widget::EWidgetMode InWidgetMode) const { ViewportInteractionCustomization->GizmoManipulationStarted(ElementWorldHandle, InWidgetMode); }
-	void GizmoManipulationDeltaUpdate(const UE::Widget::EWidgetMode InWidgetMode, const EAxisList::Type InDragAxis, const FInputDeviceState& InInputState, const FTransform& InDeltaTransform, const FVector& InPivotLocation) const { ViewportInteractionCustomization->GizmoManipulationDeltaUpdate(ElementWorldHandle, InWidgetMode, InDragAxis, InInputState, InDeltaTransform, InPivotLocation); }
+	void GizmoManipulationDeltaUpdate(const UE::Widget::EWidgetMode InWidgetMode, const EAxisList::Type InDragAxis, const FInputDeviceState& InInputState, const FTransform& InDeltaTransform, const FVector& InPivotLocation, ETypedElementViewportInteractionDragMovementType GizmoDeltaType) const { ViewportInteractionCustomization->GizmoManipulationDeltaUpdate(ElementWorldHandle, InWidgetMode, InDragAxis, InInputState, InDeltaTransform, InPivotLocation, GizmoDeltaType); }
 	void GizmoManipulationStopped(const UE::Widget::EWidgetMode InWidgetMode, const ETypedElementViewportInteractionGizmoManipulationType InManipulationType) const { ViewportInteractionCustomization->GizmoManipulationStopped(ElementWorldHandle, InWidgetMode, InManipulationType); }
 	void MirrorElement(const FVector& InMirrorScale, const FVector& InPivotLocation) const { ViewportInteractionCustomization->MirrorElement(ElementWorldHandle, InMirrorScale, InPivotLocation); }
 	bool GetFocusBounds(FBoxSphereBounds& OutBounds) const { return ViewportInteractionCustomization->GetFocusBounds(ElementWorldHandle, OutBounds); }
@@ -114,7 +128,7 @@ public:
 	 * Notify that the gizmo has manipulated the transform of the given set of elements.
 	 * @note This list should have been pre-normalized via UTypedElementSelectionSet::GetNormalizedSelection or UTypedElementSelectionSet::GetNormalizedElementList, or FLevelEditorViewportClient::GetElementsToManipulate.
 	 */
-	void UpdateGizmoManipulation(FTypedElementListConstRef InElementsToMove, const UE::Widget::EWidgetMode InWidgetMode, const EAxisList::Type InDragAxis, const FInputDeviceState& InInputState, const FTransform& InDeltaTransform);
+	void UpdateGizmoManipulation(FTypedElementListConstRef InElementsToMove, const UE::Widget::EWidgetMode InWidgetMode, const EAxisList::Type InDragAxis, const FInputDeviceState& InInputState, const FTransform& InDeltaTransform, ETypedElementViewportInteractionDragMovementType GizmoDeltaType);
 	
 	/**
 	 * Notify that the gizmo has finished manipulating the transform of the given set of elements.
diff --git a/Engine/Source/Editor/LevelEditor/Private/Elements/Actor/ActorElementLevelEditorViewportInteractionCustomization.cpp b/Engine/Source/Editor/LevelEditor/Private/Elements/Actor/ActorElementLevelEditorViewportInteractionCustomization.cpp
--- a/Engine/Source/Editor/LevelEditor/Private/Elements/Actor/ActorElementLevelEditorViewportInteractionCustomization.cpp
+++ b/Engine/Source/Editor/LevelEditor/Private/Elements/Actor/ActorElementLevelEditorViewportInteractionCustomization.cpp
@@ -30,7 +30,7 @@ void FActorElementLevelEditorViewportInteractionCustomization::GizmoManipulation
 	}
 }
 
-void FActorElementLevelEditorViewportInteractionCustomization::GizmoManipulationDeltaUpdate(const TTypedElement<ITypedElementWorldInterface>& InElementWorldHandle, const UE::Widget::EWidgetMode InWidgetMode, const EAxisList::Type InDragAxis, const FInputDeviceState& InInputState, const FTransform& InDeltaTransform, const FVector& InPivotLocation)
+void FActorElementLevelEditorViewportInteractionCustomization::GizmoManipulationDeltaUpdate(const TTypedElement<ITypedElementWorldInterface>& InElementWorldHandle, const UE::Widget::EWidgetMode InWidgetMode, const EAxisList::Type InDragAxis, const FInputDeviceState& InInputState, const FTransform& InDeltaTransform, const FVector& InPivotLocation, ETypedElementViewportInteractionDragMovementType GizmoDeltaType)
 {
 	AActor* Actor = ActorElementDataUtil::GetActorFromHandleChecked(InElementWorldHandle);
 
@@ -55,7 +55,7 @@ void FActorElementLevelEditorViewportInteractionCustomization::GizmoManipulation
 		ModifiedDeltaTransform.SetScale3D(AdjustedScale);
 	}
 
-	FActorElementEditorViewportInteractionCustomization::GizmoManipulationDeltaUpdate(InElementWorldHandle, InWidgetMode, InDragAxis, InInputState, ModifiedDeltaTransform, InPivotLocation);
+	FActorElementEditorViewportInteractionCustomization::GizmoManipulationDeltaUpdate(InElementWorldHandle, InWidgetMode, InDragAxis, InInputState, ModifiedDeltaTransform, InPivotLocation, GizmoDeltaType);
 
 	// Update the cameras from their locked actor (if any) only if the viewport is real-time enabled
 	GetMutableLevelEditorViewportClient()->UpdateLockedActorViewports(Actor, true);
diff --git a/Engine/Source/Editor/LevelEditor/Private/Elements/Component/ComponentElementLevelEditorViewportInteractionCustomization.cpp b/Engine/Source/Editor/LevelEditor/Private/Elements/Component/ComponentElementLevelEditorViewportInteractionCustomization.cpp
--- a/Engine/Source/Editor/LevelEditor/Private/Elements/Component/ComponentElementLevelEditorViewportInteractionCustomization.cpp
+++ b/Engine/Source/Editor/LevelEditor/Private/Elements/Component/ComponentElementLevelEditorViewportInteractionCustomization.cpp
@@ -32,7 +32,7 @@ void FComponentElementLevelEditorViewportInteractionCustomization::GizmoManipula
 	}
 }
 
-void FComponentElementLevelEditorViewportInteractionCustomization::GizmoManipulationDeltaUpdate(const TTypedElement<ITypedElementWorldInterface>& InElementWorldHandle, const UE::Widget::EWidgetMode InWidgetMode, const EAxisList::Type InDragAxis, const FInputDeviceState& InInputState, const FTransform& InDeltaTransform, const FVector& InPivotLocation)
+void FComponentElementLevelEditorViewportInteractionCustomization::GizmoManipulationDeltaUpdate(const TTypedElement<ITypedElementWorldInterface>& InElementWorldHandle, const UE::Widget::EWidgetMode InWidgetMode, const EAxisList::Type InDragAxis, const FInputDeviceState& InInputState, const FTransform& InDeltaTransform, const FVector& InPivotLocation, ETypedElementViewportInteractionDragMovementType GizmoDeltaType)
 {
 	UActorComponent* Component = ComponentElementDataUtil::GetComponentFromHandleChecked(InElementWorldHandle);
 
@@ -63,7 +63,7 @@ void FComponentElementLevelEditorViewportInteractionCustomization::GizmoManipula
 			ModifiedDeltaTransform.SetScale3D(AdjustedScale);
 		}
 
-		FComponentElementEditorViewportInteractionCustomization::GizmoManipulationDeltaUpdate(InElementWorldHandle, InWidgetMode, InDragAxis, InInputState, ModifiedDeltaTransform, InPivotLocation);
+		FComponentElementEditorViewportInteractionCustomization::GizmoManipulationDeltaUpdate(InElementWorldHandle, InWidgetMode, InDragAxis, InInputState, ModifiedDeltaTransform, InPivotLocation, GizmoDeltaType);
 	}
 }
 
diff --git a/Engine/Source/Editor/LevelEditor/Public/Elements/Actor/ActorElementLevelEditorViewportInteractionCustomization.h b/Engine/Source/Editor/LevelEditor/Public/Elements/Actor/ActorElementLevelEditorViewportInteractionCustomization.h
--- a/Engine/Source/Editor/LevelEditor/Public/Elements/Actor/ActorElementLevelEditorViewportInteractionCustomization.h
+++ b/Engine/Source/Editor/LevelEditor/Public/Elements/Actor/ActorElementLevelEditorViewportInteractionCustomization.h
@@ -9,7 +9,7 @@ class LEVELEDITOR_API FActorElementLevelEditorViewportInteractionCustomization :
 {
 public:
 	virtual void GizmoManipulationStarted(const TTypedElement<ITypedElementWorldInterface>& InElementWorldHandle, const UE::Widget::EWidgetMode InWidgetMode) override;
-	virtual void GizmoManipulationDeltaUpdate(const TTypedElement<ITypedElementWorldInterface>& InElementWorldHandle, const UE::Widget::EWidgetMode InWidgetMode, const EAxisList::Type InDragAxis, const FInputDeviceState& InInputState, const FTransform& InDeltaTransform, const FVector& InPivotLocation) override;
+	virtual void GizmoManipulationDeltaUpdate(const TTypedElement<ITypedElementWorldInterface>& InElementWorldHandle, const UE::Widget::EWidgetMode InWidgetMode, const EAxisList::Type InDragAxis, const FInputDeviceState& InInputState, const FTransform& InDeltaTransform, const FVector& InPivotLocation, ETypedElementViewportInteractionDragMovementType GizmoDeltaType) override;
 	virtual void GizmoManipulationStopped(const TTypedElement<ITypedElementWorldInterface>& InElementWorldHandle, const UE::Widget::EWidgetMode InWidgetMode, const ETypedElementViewportInteractionGizmoManipulationType InManipulationType) override;
 	virtual void PostGizmoManipulationStopped(TArrayView<const FTypedElementHandle> InElementHandles, const UE::Widget::EWidgetMode InWidgetMode) override;
 
diff --git a/Engine/Source/Editor/LevelEditor/Public/Elements/Component/ComponentElementLevelEditorViewportInteractionCustomization.h b/Engine/Source/Editor/LevelEditor/Public/Elements/Component/ComponentElementLevelEditorViewportInteractionCustomization.h
--- a/Engine/Source/Editor/LevelEditor/Public/Elements/Component/ComponentElementLevelEditorViewportInteractionCustomization.h
+++ b/Engine/Source/Editor/LevelEditor/Public/Elements/Component/ComponentElementLevelEditorViewportInteractionCustomization.h
@@ -11,7 +11,7 @@ class LEVELEDITOR_API FComponentElementLevelEditorViewportInteractionCustomizati
 {
 public:
 	virtual void GizmoManipulationStarted(const TTypedElement<ITypedElementWorldInterface>& InElementWorldHandle, const UE::Widget::EWidgetMode InWidgetMode) override;
-	virtual void GizmoManipulationDeltaUpdate(const TTypedElement<ITypedElementWorldInterface>& InElementWorldHandle, const UE::Widget::EWidgetMode InWidgetMode, const EAxisList::Type InDragAxis, const FInputDeviceState& InInputState, const FTransform& InDeltaTransform, const FVector& InPivotLocation) override;
+	virtual void GizmoManipulationDeltaUpdate(const TTypedElement<ITypedElementWorldInterface>& InElementWorldHandle, const UE::Widget::EWidgetMode InWidgetMode, const EAxisList::Type InDragAxis, const FInputDeviceState& InInputState, const FTransform& InDeltaTransform, const FVector& InPivotLocation, ETypedElementViewportInteractionDragMovementType GizmoDeltaType) override;
 	virtual void GizmoManipulationStopped(const TTypedElement<ITypedElementWorldInterface>& InElementWorldHandle, const UE::Widget::EWidgetMode InWidgetMode, const ETypedElementViewportInteractionGizmoManipulationType InManipulationType) override;
 
 	static void ValidateScale(const FVector& InOriginalPreDragScale, const EAxisList::Type InDragAxis, const FVector& InCurrentScale, const FVector& InBoxExtent, FVector& InOutScaleDelta, bool bInCheckSmallExtent);
diff --git a/Engine/Source/Editor/UnrealEd/Private/EditorEngine.cpp b/Engine/Source/Editor/UnrealEd/Private/EditorEngine.cpp
--- a/Engine/Source/Editor/UnrealEd/Private/EditorEngine.cpp
+++ b/Engine/Source/Editor/UnrealEd/Private/EditorEngine.cpp
@@ -2920,7 +2920,7 @@ void UEditorEngine::ApplyDeltaToActor(AActor* InActor,
 	FInputDeviceState InputState;
 	InputState.SetModifierKeyStates(bShiftDown, bAltDown, bControlDown, false);
 
-	FActorElementEditorViewportInteractionCustomization::ApplyDeltaToActor(InActor, bDelta, InTrans, InRot, InScale, GLevelEditorModeTools().PivotLocation, InputState);
+	FActorElementEditorViewportInteractionCustomization::ApplyDeltaToActor(InActor, bDelta, InTrans, InRot, InScale, GLevelEditorModeTools().PivotLocation, InputState, ETypedElementViewportInteractionDragMovementType::None);
 }
 
 void UEditorEngine::ApplyDeltaToComponent(USceneComponent* InComponent,
diff --git a/Engine/Source/Editor/UnrealEd/Private/Elements/Actor/ActorElementEditorViewportInteractionCustomization.cpp b/Engine/Source/Editor/UnrealEd/Private/Elements/Actor/ActorElementEditorViewportInteractionCustomization.cpp
--- a/Engine/Source/Editor/UnrealEd/Private/Elements/Actor/ActorElementEditorViewportInteractionCustomization.cpp
+++ b/Engine/Source/Editor/UnrealEd/Private/Elements/Actor/ActorElementEditorViewportInteractionCustomization.cpp
@@ -13,7 +13,7 @@
 #include "Elements/Interfaces/TypedElementHierarchyInterface.h"
 #include "Elements/Framework/TypedElementRegistry.h"
 
-void FActorElementEditorViewportInteractionCustomization::GizmoManipulationDeltaUpdate(const TTypedElement<ITypedElementWorldInterface>& InElementWorldHandle, const UE::Widget::EWidgetMode InWidgetMode, const EAxisList::Type InDragAxis, const FInputDeviceState& InInputState, const FTransform& InDeltaTransform, const FVector& InPivotLocation)
+void FActorElementEditorViewportInteractionCustomization::GizmoManipulationDeltaUpdate(const TTypedElement<ITypedElementWorldInterface>& InElementWorldHandle, const UE::Widget::EWidgetMode InWidgetMode, const EAxisList::Type InDragAxis, const FInputDeviceState& InInputState, const FTransform& InDeltaTransform, const FVector& InPivotLocation, ETypedElementViewportInteractionDragMovementType GizmoDeltaType)
 {
 	AActor* Actor = ActorElementDataUtil::GetActorFromHandleChecked(InElementWorldHandle);
 
@@ -21,7 +21,7 @@ void FActorElementEditorViewportInteractionCustomization::GizmoManipulationDelta
 	const FRotator DeltaRotation = InDeltaTransform.Rotator();
 	const FVector DeltaScale3D = InDeltaTransform.GetScale3D();
 
-	FActorElementEditorViewportInteractionCustomization::ApplyDeltaToActor(Actor, /*bDelta*/true, &DeltaTranslation, &DeltaRotation, &DeltaScale3D, InPivotLocation, InInputState);
+	FActorElementEditorViewportInteractionCustomization::ApplyDeltaToActor(Actor, /*bDelta*/true, &DeltaTranslation, &DeltaRotation, &DeltaScale3D, InPivotLocation, InInputState, GizmoDeltaType);
 }
 
 void FActorElementEditorViewportInteractionCustomization::MirrorElement(const TTypedElement<ITypedElementWorldInterface>& InElementWorldHandle, const FVector& InMirrorScale, const FVector& InPivotLocation)
@@ -132,7 +132,7 @@ bool FActorElementEditorViewportInteractionCustomization::GetFocusBounds( const
 	return true;
 }
 
-void FActorElementEditorViewportInteractionCustomization::ApplyDeltaToActor(AActor* InActor, const bool InIsDelta, const FVector* InDeltaTranslationPtr, const FRotator* InDeltaRotationPtr, const FVector* InDeltaScalePtr, const FVector& InPivotLocation, const FInputDeviceState& InInputState)
+void FActorElementEditorViewportInteractionCustomization::ApplyDeltaToActor(AActor* InActor, const bool InIsDelta, const FVector* InDeltaTranslationPtr, const FRotator* InDeltaRotationPtr, const FVector* InDeltaScalePtr, const FVector& InPivotLocation, const FInputDeviceState& InInputState, ETypedElementViewportInteractionDragMovementType GizmoDeltaType)
 {
 	const bool bIsSimulatingInEditor = GEditor->IsSimulatingInEditor();
 
@@ -304,7 +304,8 @@ void FActorElementEditorViewportInteractionCustomization::ApplyDeltaToActor(AAct
 		Actor->MarkPackageDirty();
 	}
 	
-	if (!GIsDemoMode)
+	// Don't need to invalidate lighting cache if this is a subsequent delta movement of the same actors
+	if (!GIsDemoMode && GizmoDeltaType != ETypedElementViewportInteractionDragMovementType::SubsequentMove)
 	{
 		for (AActor* Actor : AffectedActors)
 		{
diff --git a/Engine/Source/Editor/UnrealEd/Private/Elements/Component/ComponentElementEditorViewportInteractionCustomization.cpp b/Engine/Source/Editor/UnrealEd/Private/Elements/Component/ComponentElementEditorViewportInteractionCustomization.cpp
--- a/Engine/Source/Editor/UnrealEd/Private/Elements/Component/ComponentElementEditorViewportInteractionCustomization.cpp
+++ b/Engine/Source/Editor/UnrealEd/Private/Elements/Component/ComponentElementEditorViewportInteractionCustomization.cpp
@@ -28,7 +28,7 @@ bool FComponentElementEditorViewportInteractionCustomization::GetGizmoPivotLocat
 	return false;
 }
 
-void FComponentElementEditorViewportInteractionCustomization::GizmoManipulationDeltaUpdate(const TTypedElement<ITypedElementWorldInterface>& InElementWorldHandle, const UE::Widget::EWidgetMode InWidgetMode, const EAxisList::Type InDragAxis, const FInputDeviceState& InInputState, const FTransform& InDeltaTransform, const FVector& InPivotLocation)
+void FComponentElementEditorViewportInteractionCustomization::GizmoManipulationDeltaUpdate(const TTypedElement<ITypedElementWorldInterface>& InElementWorldHandle, const UE::Widget::EWidgetMode InWidgetMode, const EAxisList::Type InDragAxis, const FInputDeviceState& InInputState, const FTransform& InDeltaTransform, const FVector& InPivotLocation, ETypedElementViewportInteractionDragMovementType GizmoDeltaType)
 {
 	UActorComponent* Component = ComponentElementDataUtil::GetComponentFromHandleChecked(InElementWorldHandle);
 
diff --git a/Engine/Source/Editor/UnrealEd/Private/LevelEditorViewport.cpp b/Engine/Source/Editor/UnrealEd/Private/LevelEditorViewport.cpp
--- a/Engine/Source/Editor/UnrealEd/Private/LevelEditorViewport.cpp
+++ b/Engine/Source/Editor/UnrealEd/Private/LevelEditorViewport.cpp
@@ -1913,6 +1913,7 @@ FLevelEditorViewportClient::FLevelEditorViewportClient(const TSharedPtr<SLevelVi
 	, bLockedCameraView(true)
 	, bNeedToRestoreComponentBeingMovedFlag(false)
 	, bHasBegunGizmoManipulation(false)
+	, bGizmoDeltaApplied(false)
 	, bReceivedFocusRecently(false)
 	, bAlwaysShowModeWidgetAfterSelectionChanges(true)
 	, CachedElementsToManipulate(UTypedElementRegistry::GetInstance()->CreateElementList())
@@ -2849,6 +2850,7 @@ bool FLevelEditorViewportClient::InputWidgetDelta(FViewport* InViewport, EAxisLi
 									{
 										ElementsToManipulate = GetElementsToManipulate();
 										ViewportInteraction->BeginGizmoManipulation(ElementsToManipulate, GetWidgetMode());
+										bGizmoDeltaApplied = false;
 									}
 								}
 
@@ -3173,6 +3175,7 @@ void FLevelEditorViewportClient::TrackingStarted( const FInputEventState& InInpu
 		FTypedElementListConstRef ElementsToManipulate = GetElementsToManipulate();
 		ViewportInteraction->BeginGizmoManipulation(ElementsToManipulate, GetWidgetMode());
 		bHasBegunGizmoManipulation = true;
+		bGizmoDeltaApplied = false;
 
 		if (!bDuplicateActorsInProgress)
 		{
@@ -4028,8 +4031,11 @@ void FLevelEditorViewportClient::ApplyDeltaToSelectedElements(const FTransform&
 	FInputDeviceState InputState;
 	InputState.SetModifierKeyStates(IsShiftPressed(), IsAltPressed(), IsCtrlPressed(), IsCmdPressed());
 
+	ETypedElementViewportInteractionDragMovementType GizmoDeltaType = bGizmoDeltaApplied == false ? ETypedElementViewportInteractionDragMovementType::FirstMove : ETypedElementViewportInteractionDragMovementType::SubsequentMove;
+	bGizmoDeltaApplied = true;
+
 	FTypedElementListConstRef ElementsToManipulate = GetElementsToManipulate();
-	ViewportInteraction->UpdateGizmoManipulation(ElementsToManipulate, GetWidgetMode(), Widget ? Widget->GetCurrentAxis() : EAxisList::None, InputState, ModifiedDeltaTransform);
+	ViewportInteraction->UpdateGizmoManipulation(ElementsToManipulate, GetWidgetMode(), Widget ? Widget->GetCurrentAxis() : EAxisList::None, InputState, ModifiedDeltaTransform, GizmoDeltaType);
 }
 
 void FLevelEditorViewportClient::ApplyDeltaToElement(const FTypedElementHandle& InElementHandle, const FTransform& InDeltaTransform)
diff --git a/Engine/Source/Editor/UnrealEd/Public/Elements/Actor/ActorElementEditorViewportInteractionCustomization.h b/Engine/Source/Editor/UnrealEd/Public/Elements/Actor/ActorElementEditorViewportInteractionCustomization.h
--- a/Engine/Source/Editor/UnrealEd/Public/Elements/Actor/ActorElementEditorViewportInteractionCustomization.h
+++ b/Engine/Source/Editor/UnrealEd/Public/Elements/Actor/ActorElementEditorViewportInteractionCustomization.h
@@ -10,9 +10,9 @@ class AActor;
 class UNREALED_API FActorElementEditorViewportInteractionCustomization : public FTypedElementViewportInteractionCustomization
 {
 public:
-	virtual void GizmoManipulationDeltaUpdate(const TTypedElement<ITypedElementWorldInterface>& InElementWorldHandle, const UE::Widget::EWidgetMode InWidgetMode, const EAxisList::Type InDragAxis, const FInputDeviceState& InInputState, const FTransform& InDeltaTransform, const FVector& InPivotLocation) override;
+	virtual void GizmoManipulationDeltaUpdate(const TTypedElement<ITypedElementWorldInterface>& InElementWorldHandle, const UE::Widget::EWidgetMode InWidgetMode, const EAxisList::Type InDragAxis, const FInputDeviceState& InInputState, const FTransform& InDeltaTransform, const FVector& InPivotLocation, ETypedElementViewportInteractionDragMovementType GizmoDeltaType) override;
 	virtual void MirrorElement(const TTypedElement<ITypedElementWorldInterface>& InElementWorldHandle, const FVector& InMirrorScale, const FVector& InPivotLocation) override;
 	virtual bool GetFocusBounds(const TTypedElement<ITypedElementWorldInterface>& InElementWorldHandle, FBoxSphereBounds& OutBounds) override;
 	
-	static void ApplyDeltaToActor(AActor* InActor, const bool InIsDelta, const FVector* InDeltaTranslationPtr, const FRotator* InDeltaRotationPtr, const FVector* InDeltaScalePtr, const FVector& InPivotLocation, const FInputDeviceState& InInputState);
+	static void ApplyDeltaToActor(AActor* InActor, const bool InIsDelta, const FVector* InDeltaTranslationPtr, const FRotator* InDeltaRotationPtr, const FVector* InDeltaScalePtr, const FVector& InPivotLocation, const FInputDeviceState& InInputState, ETypedElementViewportInteractionDragMovementType GizmoDeltaType);
 };
diff --git a/Engine/Source/Editor/UnrealEd/Public/Elements/Component/ComponentElementEditorViewportInteractionCustomization.h b/Engine/Source/Editor/UnrealEd/Public/Elements/Component/ComponentElementEditorViewportInteractionCustomization.h
--- a/Engine/Source/Editor/UnrealEd/Public/Elements/Component/ComponentElementEditorViewportInteractionCustomization.h
+++ b/Engine/Source/Editor/UnrealEd/Public/Elements/Component/ComponentElementEditorViewportInteractionCustomization.h
@@ -12,7 +12,7 @@ class UNREALED_API FComponentElementEditorViewportInteractionCustomization : pub
 {
 public:
 	virtual bool GetGizmoPivotLocation(const TTypedElement<ITypedElementWorldInterface>& InElementWorldHandle, const UE::Widget::EWidgetMode InWidgetMode, FVector& OutPivotLocation) override;
-	virtual void GizmoManipulationDeltaUpdate(const TTypedElement<ITypedElementWorldInterface>& InElementWorldHandle, const UE::Widget::EWidgetMode InWidgetMode, const EAxisList::Type InDragAxis, const FInputDeviceState& InInputState, const FTransform& InDeltaTransform, const FVector& InPivotLocation) override;
+	virtual void GizmoManipulationDeltaUpdate(const TTypedElement<ITypedElementWorldInterface>& InElementWorldHandle, const UE::Widget::EWidgetMode InWidgetMode, const EAxisList::Type InDragAxis, const FInputDeviceState& InInputState, const FTransform& InDeltaTransform, const FVector& InPivotLocation, ETypedElementViewportInteractionDragMovementType GizmoDeltaType) override;
 	virtual bool GetFocusBounds(const TTypedElement<ITypedElementWorldInterface>& InElementWorldHandle, FBoxSphereBounds& OutBounds) override;
 	
 	static void ApplyDeltaToComponent(USceneComponent* InComponent, const bool InIsDelta, const FVector* InDeltaTranslationPtr, const FRotator* InDeltaRotationPtr, const FVector* InDeltaScalePtr, const FVector& InPivotLocation, const FInputDeviceState& InInputState);
diff --git a/Engine/Source/Editor/UnrealEd/Public/LevelEditorViewport.h b/Engine/Source/Editor/UnrealEd/Public/LevelEditorViewport.h
--- a/Engine/Source/Editor/UnrealEd/Public/LevelEditorViewport.h
+++ b/Engine/Source/Editor/UnrealEd/Public/LevelEditorViewport.h
@@ -879,6 +879,9 @@ public:
 	/** true if gizmo manipulation was started from a tracking event */
 	bool					bHasBegunGizmoManipulation;
 
+	/** set to true when a gizmo manipulation has produced a non-zero delta */
+	bool					bGizmoDeltaApplied;
+
 	/** Whether this viewport recently received focus. Used to determine whether component selection is permissible. */
 	bool bReceivedFocusRecently;
 
