Added z-index override to Viewport AddOverlay calls

diff --git a/Engine/Plugins/ChaosClothAssetEditor/Source/ChaosClothAssetEditor/Private/ChaosClothAsset/ClothEditorToolkit.cpp b/Engine/Plugins/ChaosClothAssetEditor/Source/ChaosClothAssetEditor/Private/ChaosClothAsset/ClothEditorToolkit.cpp
--- a/Engine/Plugins/ChaosClothAssetEditor/Source/ChaosClothAssetEditor/Private/ChaosClothAsset/ClothEditorToolkit.cpp
+++ b/Engine/Plugins/ChaosClothAssetEditor/Source/ChaosClothAssetEditor/Private/ChaosClothAsset/ClothEditorToolkit.cpp
@@ -366,10 +366,10 @@ TSharedPtr<FEditorViewportClient> FChaosClothAssetEditorToolkit::CreateEditorVie
 
 //~ Begin FAssetEditorToolkit overrides
 
-void FChaosClothAssetEditorToolkit::AddViewportOverlayWidget(TSharedRef<SWidget> InViewportOverlayWidget)
+void FChaosClothAssetEditorToolkit::AddViewportOverlayWidget(TSharedRef<SWidget> InViewportOverlayWidget, int32 ZOrder)
 {
 	TSharedPtr<SChaosClothAssetEditorRestSpaceViewport> ViewportWidget = StaticCastSharedPtr<SChaosClothAssetEditorRestSpaceViewport>(ViewportTabContent->GetFirstViewport());
-	ViewportWidget->AddOverlayWidget(InViewportOverlayWidget);
+	ViewportWidget->AddOverlayWidget(InViewportOverlayWidget, ZOrder);
 }
 
 void FChaosClothAssetEditorToolkit::RemoveViewportOverlayWidget(TSharedRef<SWidget> InViewportOverlayWidget)
diff --git a/Engine/Plugins/ChaosClothAssetEditor/Source/ChaosClothAssetEditor/Private/ChaosClothAsset/ClothEditorToolkit.h b/Engine/Plugins/ChaosClothAssetEditor/Source/ChaosClothAssetEditor/Private/ChaosClothAsset/ClothEditorToolkit.h
--- a/Engine/Plugins/ChaosClothAssetEditor/Source/ChaosClothAssetEditor/Private/ChaosClothAsset/ClothEditorToolkit.h
+++ b/Engine/Plugins/ChaosClothAssetEditor/Source/ChaosClothAssetEditor/Private/ChaosClothAsset/ClothEditorToolkit.h
@@ -80,7 +80,7 @@ private:
 	virtual TSharedPtr<FEditorViewportClient> CreateEditorViewportClient() const override;
 
 	// FAssetEditorToolkit
-	virtual void AddViewportOverlayWidget(TSharedRef<SWidget> InViewportOverlayWidget) override;
+	virtual void AddViewportOverlayWidget(TSharedRef<SWidget> InViewportOverlayWidget, int32 ZOrder = INDEX_NONE) override;
 	virtual void RemoveViewportOverlayWidget(TSharedRef<SWidget> InViewportOverlayWidget) override;
 	virtual bool OnRequestClose(EAssetEditorCloseReason InCloseReason) override;
 	virtual void OnClose() override;
diff --git a/Engine/Plugins/Editor/UVEditor/Source/UVEditor/Private/SUVEditor2DViewport.cpp b/Engine/Plugins/Editor/UVEditor/Source/UVEditor/Private/SUVEditor2DViewport.cpp
--- a/Engine/Plugins/Editor/UVEditor/Source/UVEditor/Private/SUVEditor2DViewport.cpp
+++ b/Engine/Plugins/Editor/UVEditor/Source/UVEditor/Private/SUVEditor2DViewport.cpp
@@ -75,9 +75,9 @@ void SUVEditor2DViewport::BindCommands()
 		}));
 }
 
-void SUVEditor2DViewport::AddOverlayWidget(TSharedRef<SWidget> OverlaidWidget)
+void SUVEditor2DViewport::AddOverlayWidget(TSharedRef<SWidget> OverlaidWidget, int32 ZOrder)
 {
-	ViewportOverlay->AddSlot()
+	ViewportOverlay->AddSlot(ZOrder)
 	[
 		OverlaidWidget
 	];
diff --git a/Engine/Plugins/Editor/UVEditor/Source/UVEditor/Private/UVEditorToolkit.cpp b/Engine/Plugins/Editor/UVEditor/Source/UVEditor/Private/UVEditorToolkit.cpp
--- a/Engine/Plugins/Editor/UVEditor/Source/UVEditor/Private/UVEditorToolkit.cpp
+++ b/Engine/Plugins/Editor/UVEditor/Source/UVEditor/Private/UVEditorToolkit.cpp
@@ -304,10 +304,10 @@ void FUVEditorToolkit::OnClose()
 
 // These get called indirectly (via toolkit host) from the mode toolkit when the mode starts or ends a tool,
 // in order to add or remove an accept/cancel overlay.
-void FUVEditorToolkit::AddViewportOverlayWidget(TSharedRef<SWidget> InViewportOverlayWidget) 
+void FUVEditorToolkit::AddViewportOverlayWidget(TSharedRef<SWidget> InViewportOverlayWidget, int32 ZOrder) 
 {
 	TSharedPtr<SUVEditor2DViewport> ViewportWidget = StaticCastSharedPtr<SUVEditor2DViewport>(ViewportTabContent->GetFirstViewport());
-	ViewportWidget->AddOverlayWidget(InViewportOverlayWidget);
+	ViewportWidget->AddOverlayWidget(InViewportOverlayWidget, ZOrder);
 }
 void FUVEditorToolkit::RemoveViewportOverlayWidget(TSharedRef<SWidget> InViewportOverlayWidget)
 {
diff --git a/Engine/Plugins/Editor/UVEditor/Source/UVEditor/Public/SUVEditor2DViewport.h b/Engine/Plugins/Editor/UVEditor/Source/UVEditor/Public/SUVEditor2DViewport.h
--- a/Engine/Plugins/Editor/UVEditor/Source/UVEditor/Public/SUVEditor2DViewport.h
+++ b/Engine/Plugins/Editor/UVEditor/Source/UVEditor/Public/SUVEditor2DViewport.h
@@ -14,7 +14,7 @@ public:
 
 	// These allow the toolkit to add an accept/cancel overlay when needed. PopulateViewportOverlays
 	// is not helpful here because that gets called just once.
-	virtual void AddOverlayWidget(TSharedRef<SWidget> OverlaidWidget);
+	virtual void AddOverlayWidget(TSharedRef<SWidget> OverlaidWidget, int32 ZOrder = INDEX_NONE);
 	virtual void RemoveOverlayWidget(TSharedRef<SWidget> OverlaidWidget);
 
 	// SEditorViewport
diff --git a/Engine/Plugins/Editor/UVEditor/Source/UVEditor/Public/UVEditorToolkit.h b/Engine/Plugins/Editor/UVEditor/Source/UVEditor/Public/UVEditorToolkit.h
--- a/Engine/Plugins/Editor/UVEditor/Source/UVEditor/Public/UVEditorToolkit.h
+++ b/Engine/Plugins/Editor/UVEditor/Source/UVEditor/Public/UVEditorToolkit.h
@@ -41,7 +41,7 @@ public:
 	virtual void CreateWidgets() override;
 
 	// FAssetEditorToolkit
-	virtual void AddViewportOverlayWidget(TSharedRef<SWidget> InViewportOverlayWidget) override;
+	virtual void AddViewportOverlayWidget(TSharedRef<SWidget> InViewportOverlayWidget, int32 ZOrder = INDEX_NONE) override;
 	virtual void RemoveViewportOverlayWidget(TSharedRef<SWidget> InViewportOverlayWidget) override;
 	virtual void CreateEditorModeManager() override;
 	virtual FText GetToolkitName() const override;
diff --git a/Engine/Plugins/Experimental/CharacterFXEditor/BaseCharacterFXEditor/Source/BaseCharacterFXEditor/Private/BaseCharacterFXEditorToolkit.cpp b/Engine/Plugins/Experimental/CharacterFXEditor/BaseCharacterFXEditor/Source/BaseCharacterFXEditor/Private/BaseCharacterFXEditorToolkit.cpp
--- a/Engine/Plugins/Experimental/CharacterFXEditor/BaseCharacterFXEditor/Source/BaseCharacterFXEditor/Private/BaseCharacterFXEditorToolkit.cpp
+++ b/Engine/Plugins/Experimental/CharacterFXEditor/BaseCharacterFXEditor/Source/BaseCharacterFXEditor/Private/BaseCharacterFXEditorToolkit.cpp
@@ -118,10 +118,10 @@ AssetEditorViewportFactoryFunction FBaseCharacterFXEditorToolkit::GetViewportDel
 	return TempViewportDelegate;
 }
 
-void FBaseCharacterFXEditorToolkit::AddViewportOverlayWidget(TSharedRef<SWidget> InViewportOverlayWidget)
+void FBaseCharacterFXEditorToolkit::AddViewportOverlayWidget(TSharedRef<SWidget> InViewportOverlayWidget, int32 ZOrder)
 {
 	TSharedPtr<SBaseCharacterFXEditorViewport> ViewportWidget = StaticCastSharedPtr<SBaseCharacterFXEditorViewport>(ViewportTabContent->GetFirstViewport());
-	ViewportWidget->AddOverlayWidget(InViewportOverlayWidget);
+	ViewportWidget->AddOverlayWidget(InViewportOverlayWidget, ZOrder);
 }
 
 void FBaseCharacterFXEditorToolkit::RemoveViewportOverlayWidget(TSharedRef<SWidget> InViewportOverlayWidget)
diff --git a/Engine/Plugins/Experimental/CharacterFXEditor/BaseCharacterFXEditor/Source/BaseCharacterFXEditor/Private/SBaseCharacterFXEditorViewport.cpp b/Engine/Plugins/Experimental/CharacterFXEditor/BaseCharacterFXEditor/Source/BaseCharacterFXEditor/Private/SBaseCharacterFXEditorViewport.cpp
--- a/Engine/Plugins/Experimental/CharacterFXEditor/BaseCharacterFXEditor/Source/BaseCharacterFXEditor/Private/SBaseCharacterFXEditorViewport.cpp
+++ b/Engine/Plugins/Experimental/CharacterFXEditor/BaseCharacterFXEditor/Source/BaseCharacterFXEditor/Private/SBaseCharacterFXEditorViewport.cpp
@@ -5,9 +5,9 @@
 
 #define LOCTEXT_NAMESPACE "SBaseCharacterFXEditorViewport"
 
-void SBaseCharacterFXEditorViewport::AddOverlayWidget(TSharedRef<SWidget> OverlaidWidget)
+void SBaseCharacterFXEditorViewport::AddOverlayWidget(TSharedRef<SWidget> OverlaidWidget, int32 ZOrder)
 {
-	ViewportOverlay->AddSlot()
+	ViewportOverlay->AddSlot(ZOrder)
 	[
 		OverlaidWidget
 	];
diff --git a/Engine/Plugins/Experimental/CharacterFXEditor/BaseCharacterFXEditor/Source/BaseCharacterFXEditor/Public/BaseCharacterFXEditorToolkit.h b/Engine/Plugins/Experimental/CharacterFXEditor/BaseCharacterFXEditor/Source/BaseCharacterFXEditor/Public/BaseCharacterFXEditorToolkit.h
--- a/Engine/Plugins/Experimental/CharacterFXEditor/BaseCharacterFXEditor/Source/BaseCharacterFXEditor/Public/BaseCharacterFXEditorToolkit.h
+++ b/Engine/Plugins/Experimental/CharacterFXEditor/BaseCharacterFXEditor/Source/BaseCharacterFXEditor/Public/BaseCharacterFXEditorToolkit.h
@@ -30,7 +30,7 @@ public:
 	virtual void RegisterTabSpawners(const TSharedRef<FTabManager>& InTabManager) override;
 
 	virtual AssetEditorViewportFactoryFunction GetViewportDelegate() override;
-	virtual void AddViewportOverlayWidget(TSharedRef<SWidget> InViewportOverlayWidget) override;
+	virtual void AddViewportOverlayWidget(TSharedRef<SWidget> InViewportOverlayWidget, int32 ZOrder = INDEX_NONE) override;
 	virtual void RemoveViewportOverlayWidget(TSharedRef<SWidget> InViewportOverlayWidget) override;
 
 protected:
diff --git a/Engine/Plugins/Experimental/CharacterFXEditor/BaseCharacterFXEditor/Source/BaseCharacterFXEditor/Public/SBaseCharacterFXEditorViewport.h b/Engine/Plugins/Experimental/CharacterFXEditor/BaseCharacterFXEditor/Source/BaseCharacterFXEditor/Public/SBaseCharacterFXEditorViewport.h
--- a/Engine/Plugins/Experimental/CharacterFXEditor/BaseCharacterFXEditor/Source/BaseCharacterFXEditor/Public/SBaseCharacterFXEditorViewport.h
+++ b/Engine/Plugins/Experimental/CharacterFXEditor/BaseCharacterFXEditor/Source/BaseCharacterFXEditor/Public/SBaseCharacterFXEditorViewport.h
@@ -11,7 +11,7 @@ public:
 
 	// These allow the toolkit to add an accept/cancel overlay when needed. PopulateViewportOverlays
 	// is not helpful here because that gets called just once.
-	virtual void AddOverlayWidget(TSharedRef<SWidget> OverlaidWidget);
+	virtual void AddOverlayWidget(TSharedRef<SWidget> OverlaidWidget, int32 ZOrder = INDEX_NONE);
 	virtual void RemoveOverlayWidget(TSharedRef<SWidget> OverlaidWidget);
 
 	// SEditorViewport
diff --git a/Engine/Source/Editor/EditorFramework/Public/Toolkits/IToolkitHost.h b/Engine/Source/Editor/EditorFramework/Public/Toolkits/IToolkitHost.h
--- a/Engine/Source/Editor/EditorFramework/Public/Toolkits/IToolkitHost.h
+++ b/Engine/Source/Editor/EditorFramework/Public/Toolkits/IToolkitHost.h
@@ -51,7 +51,15 @@ public:
 
 	/** Allows Toolkits to push widgets to the viewport.  Passing a nullptr for the Viewport will add  or 
 	    remove the OverlayWidget to or from the ActiveViewport */
-	virtual void AddViewportOverlayWidget(TSharedRef<SWidget> InOverlaidWidget, TSharedPtr<IAssetViewport> InViewport = nullptr)  {}
+	virtual void AddViewportOverlayWidget(TSharedRef<SWidget> InOverlaidWidget, TSharedPtr<IAssetViewport> InViewport = nullptr)
+	{
+		AddViewportOverlayWidget(InOverlaidWidget, INDEX_NONE, InViewport);		
+	}
+
+	/** Allows Toolkits to push widgets to the viewport with a specific Z-Order.  Passing a nullptr for the Viewport will add  or 
+		remove the OverlayWidget to or from the ActiveViewport */
+	virtual void AddViewportOverlayWidget(TSharedRef<SWidget> InOverlaidWidget, int32 ZOrder, TSharedPtr<IAssetViewport> InViewport = nullptr)  {}
+
 	virtual void RemoveViewportOverlayWidget(TSharedRef<SWidget> InOverlaidWidget, TSharedPtr<IAssetViewport> InViewport = nullptr) {}
 
 	/** Return the size of the active viewport, accounting for DPI Scale */
diff --git a/Engine/Source/Editor/LevelEditor/Private/SLevelEditor.cpp b/Engine/Source/Editor/LevelEditor/Private/SLevelEditor.cpp
--- a/Engine/Source/Editor/LevelEditor/Private/SLevelEditor.cpp
+++ b/Engine/Source/Editor/LevelEditor/Private/SLevelEditor.cpp
@@ -2227,15 +2227,15 @@ FName SLevelEditor::GetStatusBarName() const
 	return LevelEditorStatusBarName;
 }
 
-void SLevelEditor::AddViewportOverlayWidget(TSharedRef<SWidget> InWidget, TSharedPtr<IAssetViewport> InViewport)
+void SLevelEditor::AddViewportOverlayWidget(TSharedRef<SWidget> InWidget, int32 ZOrder, TSharedPtr<IAssetViewport> InViewport)
 {
 	if (InViewport != nullptr)
 	{
-		InViewport->AddOverlayWidget(InWidget);
+		InViewport->AddOverlayWidget(InWidget, ZOrder);
 	}
 	else if (TSharedPtr<SLevelViewport> ActiveViewport = GetActiveViewport())
 	{
-		ActiveViewport->AddOverlayWidget(InWidget);
+		ActiveViewport->AddOverlayWidget(InWidget, ZOrder);
 	}
 }
 
diff --git a/Engine/Source/Editor/LevelEditor/Private/SLevelEditor.h b/Engine/Source/Editor/LevelEditor/Private/SLevelEditor.h
--- a/Engine/Source/Editor/LevelEditor/Private/SLevelEditor.h
+++ b/Engine/Source/Editor/LevelEditor/Private/SLevelEditor.h
@@ -136,7 +136,7 @@ public:
 	virtual UTypedElementCommonActions* GetCommonActions() const override;
 	virtual FName GetStatusBarName() const override;
 	virtual FOnActiveViewportChanged& OnActiveViewportChanged() { return OnActiveViewportChangedDelegate; }
-	virtual void AddViewportOverlayWidget(TSharedRef<SWidget>, TSharedPtr<IAssetViewport> InViewport = nullptr) override;
+	virtual void AddViewportOverlayWidget(TSharedRef<SWidget>, int32 ZOrder = INDEX_NONE, TSharedPtr<IAssetViewport> InViewport = nullptr) override;
 	virtual void RemoveViewportOverlayWidget(TSharedRef<SWidget>, TSharedPtr<IAssetViewport> InViewport = nullptr) override; 
 
 
diff --git a/Engine/Source/Editor/LevelEditor/Private/SLevelViewport.cpp b/Engine/Source/Editor/LevelEditor/Private/SLevelViewport.cpp
--- a/Engine/Source/Editor/LevelEditor/Private/SLevelViewport.cpp
+++ b/Engine/Source/Editor/LevelEditor/Private/SLevelViewport.cpp
@@ -4679,9 +4679,9 @@ void SLevelViewport::RemoveActorPreview( int32 PreviewIndex, AActor* Actor, cons
 	}
 }
 
-void SLevelViewport::AddOverlayWidget(TSharedRef<SWidget> OverlaidWidget)
+void SLevelViewport::AddOverlayWidget(TSharedRef<SWidget> OverlaidWidget, int32 ZOrder)
 {
-	ViewportOverlay->AddSlot()
+	ViewportOverlay->AddSlot(ZOrder)
 	[
 		OverlaidWidget
 	];
diff --git a/Engine/Source/Editor/LevelEditor/Public/SLevelViewport.h b/Engine/Source/Editor/LevelEditor/Public/SLevelViewport.h
--- a/Engine/Source/Editor/LevelEditor/Public/SLevelViewport.h
+++ b/Engine/Source/Editor/LevelEditor/Public/SLevelViewport.h
@@ -118,7 +118,7 @@ public:
 	virtual TSharedRef< const SWidget> AsWidget() const override { return AsShared(); }
 	virtual TSharedRef< SWidget> AsWidget() override { return AsShared(); }
 	virtual TWeakPtr< SViewport > GetViewportWidget() override { return ViewportWidget; }
-	virtual void AddOverlayWidget( TSharedRef<SWidget> OverlaidWidget ) override;
+	virtual void AddOverlayWidget( TSharedRef<SWidget> OverlaidWidget, int32 ZOrder=INDEX_NONE ) override;
 	virtual void RemoveOverlayWidget( TSharedRef<SWidget> OverlaidWidget ) override;
 
 
diff --git a/Engine/Source/Editor/Persona/Private/SAnimationEditorViewport.cpp b/Engine/Source/Editor/Persona/Private/SAnimationEditorViewport.cpp
--- a/Engine/Source/Editor/Persona/Private/SAnimationEditorViewport.cpp
+++ b/Engine/Source/Editor/Persona/Private/SAnimationEditorViewport.cpp
@@ -483,9 +483,9 @@ void SAnimationEditorViewportTabBody::AddToolbarExtender(FName MenuToExtend, FMe
 	return ViewportWidget->ViewportToolbar->AddMenuExtender(MenuToExtend, MenuBuilderDelegate);
 }
 
-void SAnimationEditorViewportTabBody::AddOverlayWidget(TSharedRef<SWidget> InOverlaidWidget)
+void SAnimationEditorViewportTabBody::AddOverlayWidget(TSharedRef<SWidget> InOverlaidWidget, int32 ZOrder)
 {
-	ViewportWidget->ViewportOverlay->AddSlot()
+	ViewportWidget->ViewportOverlay->AddSlot(ZOrder)
 	[
 		InOverlaidWidget
 	];
diff --git a/Engine/Source/Editor/Persona/Private/SAnimationEditorViewport.h b/Engine/Source/Editor/Persona/Private/SAnimationEditorViewport.h
--- a/Engine/Source/Editor/Persona/Private/SAnimationEditorViewport.h
+++ b/Engine/Source/Editor/Persona/Private/SAnimationEditorViewport.h
@@ -226,7 +226,7 @@ public:
 	virtual void RemoveNotification(const TWeakPtr<SWidget>& InContainingWidget) override;
 	virtual void AddToolbarExtender(FName MenuToExtend, FMenuExtensionDelegate MenuBuilderDelegate) override;
 	virtual FPersonaViewportKeyDownDelegate& GetKeyDownDelegate() override { return OnKeyDownDelegate; }
-	virtual void AddOverlayWidget( TSharedRef<SWidget> InOverlaidWidget ) override;
+	virtual void AddOverlayWidget( TSharedRef<SWidget> InOverlaidWidget, int32 ZOrder = INDEX_NONE ) override;
 	virtual void RemoveOverlayWidget( TSharedRef<SWidget> InOverlaidWidget ) override;
 
 	
diff --git a/Engine/Source/Editor/Persona/Public/IPersonaViewport.h b/Engine/Source/Editor/Persona/Public/IPersonaViewport.h
--- a/Engine/Source/Editor/Persona/Public/IPersonaViewport.h
+++ b/Engine/Source/Editor/Persona/Public/IPersonaViewport.h
@@ -74,7 +74,7 @@ public:
 	virtual FPersonaViewportKeyDownDelegate& GetKeyDownDelegate() = 0;
 
 	/** Overlay a widget over the whole viewport */
-	virtual void AddOverlayWidget( TSharedRef<SWidget> InOverlaidWidget ) = 0;
+	virtual void AddOverlayWidget( TSharedRef<SWidget> InOverlaidWidget, int32 ZOrder=INDEX_NONE ) = 0;
 
 	/** Remove an overlay widget from the viewport */
 	virtual void RemoveOverlayWidget( TSharedRef<SWidget> InOverlaidWidget ) = 0;
diff --git a/Engine/Source/Editor/SkeletalMeshEditor/Private/SkeletalMeshEditor.cpp b/Engine/Source/Editor/SkeletalMeshEditor/Private/SkeletalMeshEditor.cpp
--- a/Engine/Source/Editor/SkeletalMeshEditor/Private/SkeletalMeshEditor.cpp
+++ b/Engine/Source/Editor/SkeletalMeshEditor/Private/SkeletalMeshEditor.cpp
@@ -622,11 +622,11 @@ void FSkeletalMeshEditor::InitToolMenuContext(FToolMenuContext& MenuContext)
 	MenuContext.AddObject(PersonaContext);
 }
 
-void FSkeletalMeshEditor::AddViewportOverlayWidget(TSharedRef<SWidget> InOverlaidWidget)
+void FSkeletalMeshEditor::AddViewportOverlayWidget(TSharedRef<SWidget> InOverlaidWidget, int32 ZOrder)
 {
 	if (Viewport.IsValid())
 	{
-		Viewport->AddOverlayWidget(InOverlaidWidget);
+		Viewport->AddOverlayWidget(InOverlaidWidget, ZOrder);
 	}	
 }
 
diff --git a/Engine/Source/Editor/SkeletalMeshEditor/Private/SkeletalMeshEditor.h b/Engine/Source/Editor/SkeletalMeshEditor/Private/SkeletalMeshEditor.h
--- a/Engine/Source/Editor/SkeletalMeshEditor/Private/SkeletalMeshEditor.h
+++ b/Engine/Source/Editor/SkeletalMeshEditor/Private/SkeletalMeshEditor.h
@@ -72,7 +72,7 @@ public:
 	virtual FLinearColor GetWorldCentricTabColorScale() const override;
 	virtual void InitToolMenuContext(FToolMenuContext& MenuContext) override;
 
-	virtual void AddViewportOverlayWidget(TSharedRef<SWidget> InOverlaidWidget) override;
+	virtual void AddViewportOverlayWidget(TSharedRef<SWidget> InOverlaidWidget, int32 ZOrder = INDEX_NONE) override;
 	virtual void RemoveViewportOverlayWidget(TSharedRef<SWidget> InOverlaidWidget) override;
 	
 	/** FBaseToolkit overrides */
diff --git a/Engine/Source/Editor/StaticMeshEditor/Private/StaticMeshEditor.cpp b/Engine/Source/Editor/StaticMeshEditor/Private/StaticMeshEditor.cpp
--- a/Engine/Source/Editor/StaticMeshEditor/Private/StaticMeshEditor.cpp
+++ b/Engine/Source/Editor/StaticMeshEditor/Private/StaticMeshEditor.cpp
@@ -2754,13 +2754,13 @@ TStatId FStaticMeshEditor::GetStatId() const
 	RETURN_QUICK_DECLARE_CYCLE_STAT(FStaticMeshEditor, STATGROUP_TaskGraphTasks);
 }
 
-void FStaticMeshEditor::AddViewportOverlayWidget(TSharedRef<SWidget> InOverlaidWidget)
+void FStaticMeshEditor::AddViewportOverlayWidget(TSharedRef<SWidget> InOverlaidWidget, int32 ZOrder)
 {
 	TSharedPtr<SStaticMeshEditorViewport> Viewport = GetStaticMeshViewport();
 
 	if (Viewport.IsValid() && Viewport->GetViewportOverlay().IsValid())
 	{
-		Viewport->GetViewportOverlay()->AddSlot()
+		Viewport->GetViewportOverlay()->AddSlot(ZOrder)
 		[
 			InOverlaidWidget
 		];
diff --git a/Engine/Source/Editor/StaticMeshEditor/Private/StaticMeshEditor.h b/Engine/Source/Editor/StaticMeshEditor/Private/StaticMeshEditor.h
--- a/Engine/Source/Editor/StaticMeshEditor/Private/StaticMeshEditor.h
+++ b/Engine/Source/Editor/StaticMeshEditor/Private/StaticMeshEditor.h
@@ -227,7 +227,7 @@ public:
 	virtual TStatId GetStatId() const final;
 
 	/** Add a widget to the StaticMeshViewport's ViewportOverlay */
-	void AddViewportOverlayWidget(TSharedRef<SWidget> InOverlaidWidget) override;
+	void AddViewportOverlayWidget(TSharedRef<SWidget> InOverlaidWidget, int32 ZOrder = INDEX_NONE) override;
 
 	/** Remove a widget from the StaticMeshViewport's ViewportOverlay */
 	void RemoveViewportOverlayWidget(TSharedRef<SWidget> InViewportOverlayWidget) override;
diff --git a/Engine/Source/Editor/UnrealEd/Private/Toolkits/SStandaloneAssetEditorToolkitHost.cpp b/Engine/Source/Editor/UnrealEd/Private/Toolkits/SStandaloneAssetEditorToolkitHost.cpp
--- a/Engine/Source/Editor/UnrealEd/Private/Toolkits/SStandaloneAssetEditorToolkitHost.cpp
+++ b/Engine/Source/Editor/UnrealEd/Private/Toolkits/SStandaloneAssetEditorToolkitHost.cpp
@@ -393,11 +393,11 @@ UTypedElementCommonActions* SStandaloneAssetEditorToolkitHost::GetCommonActions(
 	return CommonActions.Get();
 }
 
-void SStandaloneAssetEditorToolkitHost::AddViewportOverlayWidget(TSharedRef<SWidget> InOverlaidWidget, TSharedPtr<IAssetViewport> InViewport)
+void SStandaloneAssetEditorToolkitHost::AddViewportOverlayWidget(TSharedRef<SWidget> InOverlaidWidget, int32 ZOrder, TSharedPtr<IAssetViewport> InViewport)
 {
 	if (HostedAssetEditorToolkit.IsValid())
 	{
-		HostedAssetEditorToolkit->AddViewportOverlayWidget(InOverlaidWidget);
+		HostedAssetEditorToolkit->AddViewportOverlayWidget(InOverlaidWidget, ZOrder);
 	}
 }
 
diff --git a/Engine/Source/Editor/UnrealEd/Private/Toolkits/SStandaloneAssetEditorToolkitHost.h b/Engine/Source/Editor/UnrealEd/Private/Toolkits/SStandaloneAssetEditorToolkitHost.h
--- a/Engine/Source/Editor/UnrealEd/Private/Toolkits/SStandaloneAssetEditorToolkitHost.h
+++ b/Engine/Source/Editor/UnrealEd/Private/Toolkits/SStandaloneAssetEditorToolkitHost.h
@@ -61,7 +61,7 @@ public:
 	virtual UWorld* GetWorld() const override;
 	virtual UTypedElementCommonActions* GetCommonActions() const override;
 	FName GetStatusBarName() const override { return StatusBarName; } 
-	virtual void AddViewportOverlayWidget(TSharedRef<SWidget> InOverlaidWidget, TSharedPtr<IAssetViewport> InViewport = nullptr) override;
+	virtual void AddViewportOverlayWidget(TSharedRef<SWidget> InOverlaidWidget, int32 ZOrder = INDEX_NONE, TSharedPtr<IAssetViewport> InViewport = nullptr) override;
 	virtual void RemoveViewportOverlayWidget(TSharedRef<SWidget> InOverlaidWidget, TSharedPtr<IAssetViewport> InViewport = nullptr) override;
 	virtual FOnActiveViewportChanged& OnActiveViewportChanged() override { return OnActiveViewportChangedDelegate; }
 	void CreateDefaultStandaloneMenuBar(UToolMenu* MenuBar);
diff --git a/Engine/Source/Editor/UnrealEd/Public/IAssetViewport.h b/Engine/Source/Editor/UnrealEd/Public/IAssetViewport.h
--- a/Engine/Source/Editor/UnrealEd/Public/IAssetViewport.h
+++ b/Engine/Source/Editor/UnrealEd/Public/IAssetViewport.h
@@ -95,7 +95,7 @@ public:
 	virtual bool IsInGameView() const = 0;
 	
 	/** Adds a widget overlaid over the viewport */
-	virtual void AddOverlayWidget(TSharedRef<SWidget> OverlaidWidget) = 0;
+	virtual void AddOverlayWidget(TSharedRef<SWidget> OverlaidWidget, int32 ZOrder=INDEX_NONE) = 0;
 
 	/** Removes a widget that was previously overlaid on to this viewport */
 	virtual void RemoveOverlayWidget(TSharedRef<SWidget> OverlaidWidget) = 0;
diff --git a/Engine/Source/Editor/UnrealEd/Public/Toolkits/AssetEditorToolkit.h b/Engine/Source/Editor/UnrealEd/Public/Toolkits/AssetEditorToolkit.h
--- a/Engine/Source/Editor/UnrealEd/Public/Toolkits/AssetEditorToolkit.h
+++ b/Engine/Source/Editor/UnrealEd/Public/Toolkits/AssetEditorToolkit.h
@@ -246,10 +246,10 @@ public:
 	// Called when another toolkit (such as a ed mode toolkit) is no longer being hosted in this asset editor toolkit
 	virtual void OnToolkitHostingFinished(const TSharedRef<IToolkit>& Toolkit) {}
 
-	/* Called when a toolkit requests an overlay widget to be added to the viewport. Not relevant in the absence
+	/* Called when a toolkit requests an overlay widget to be added to the viewport, optionally with a specific Z-Order. Not relevant in the absence
 	 * of a viewport.
 	 */
-	virtual void AddViewportOverlayWidget(TSharedRef<SWidget> InViewportOverlayWidget) {}
+	virtual void AddViewportOverlayWidget(TSharedRef<SWidget> InViewportOverlayWidget, int32 ZOrder = INDEX_NONE) {}
 
 	/** Called when a toolkit requests the overlay widget to be removed. */
 	virtual void RemoveViewportOverlayWidget(TSharedRef<SWidget> InViewportOverlayWidget) {}
