Add an Editor.UseForwardRightUpDisplayNames console variable.

This cvar will be used to control display names in details panel customizations, viewport information, etc...  This cvar is ECVF_ReadOnly so that it cannot be changed by the user from the console.

diff --git a/Engine/Source/Editor/EditorFramework/Private/EditorFrameworkModule.cpp b/Engine/Source/Editor/EditorFramework/Private/EditorFrameworkModule.cpp
--- a/Engine/Source/Editor/EditorFramework/Private/EditorFrameworkModule.cpp
+++ b/Engine/Source/Editor/EditorFramework/Private/EditorFrameworkModule.cpp
@@ -2,8 +2,19 @@
 
 #include "EditorFrameworkModule.h"
 
+#include "HAL/IConsoleManager.h"
 #include "Modules/ModuleManager.h"
 
+namespace UE::Editor::Private
+{
+	static TAutoConsoleVariable<bool> EditorUseForwardRightUpDisplayNames(
+		TEXT("Editor.UseForwardRightUpDisplayNames"),
+		false,
+		TEXT("When true, the editor uses (Forward, Right, Up) nomenclature for coordinate axes"),
+		ECVF_ReadOnly
+	);
+}
+
 void FEditorFrameworkModule::StartupModule()
 {
 }
