[Truncation Warnings] Fix warnings in EditorFramework module

#rnx

diff --git a/Engine/Source/Editor/EditorFramework/Private/Viewports/InViewportUIDragOperation.cpp b/Engine/Source/Editor/EditorFramework/Private/Viewports/InViewportUIDragOperation.cpp
--- a/Engine/Source/Editor/EditorFramework/Private/Viewports/InViewportUIDragOperation.cpp
+++ b/Engine/Source/Editor/EditorFramework/Private/Viewports/InViewportUIDragOperation.cpp
@@ -13,7 +13,7 @@
 class SWidget;
 struct FPointerEvent;
 
-const FVector2D FInViewportUIDragOperation::GetDecoratorOffsetFromCursor()
+const FVector2f FInViewportUIDragOperation::GetDecoratorOffsetFromCursor()
 {
 	return DecoratorOffsetFromCursor;
 }
@@ -36,7 +36,7 @@ void FInViewportUIDragOperation::OnDrop(bool bDropWasHandled, const FPointerEven
 void FInViewportUIDragOperation::OnDragged(const FDragDropEvent& DragDropEvent)
 {
 	// The tab is being dragged. Move the the decorator window to match the cursor position.
-	FVector2D TargetPosition = DragDropEvent.GetScreenSpacePosition() - GetDecoratorOffsetFromCursor();
+	FVector2f TargetPosition = DragDropEvent.GetScreenSpacePosition() - GetDecoratorOffsetFromCursor();
 	CursorDecoratorWindow->UpdateMorphTargetShape(FSlateRect(TargetPosition.X, TargetPosition.Y, TargetPosition.X + LastContentSize.X, TargetPosition.Y + LastContentSize.Y));
 	CursorDecoratorWindow->MoveWindowTo(TargetPosition);
 }
@@ -54,8 +54,8 @@ FInViewportUIDragOperation::~FInViewportUIDragOperation()
 
 FInViewportUIDragOperation::FInViewportUIDragOperation(const TSharedRef<class SWidget>& InUIToBeDragged, const FVector2D InDecoratorOffset, const FVector2D& OwnerAreaSize, FOnInViewportUIDropped& InDropDelegate)
 	: UIBeingDragged(InUIToBeDragged)
-	, DecoratorOffsetFromCursor(InDecoratorOffset)
-	, LastContentSize(OwnerAreaSize)
+	, DecoratorOffsetFromCursor(UE::Slate::CastToVector2f(InDecoratorOffset))
+	, LastContentSize(UE::Slate::CastToVector2f(OwnerAreaSize))
 	, DropDelegate(InDropDelegate)
 {
 	// Create the decorator window that we will use during this drag and drop to make the user feel like
diff --git a/Engine/Source/Editor/EditorFramework/Public/Viewports/InViewportUIDragOperation.h b/Engine/Source/Editor/EditorFramework/Public/Viewports/InViewportUIDragOperation.h
--- a/Engine/Source/Editor/EditorFramework/Public/Viewports/InViewportUIDragOperation.h
+++ b/Engine/Source/Editor/EditorFramework/Public/Viewports/InViewportUIDragOperation.h
@@ -52,7 +52,7 @@ public:
 	virtual ~FInViewportUIDragOperation();
 
 	/** @return The offset into the tab where the user grabbed in Slate Units. */
-	const FVector2D GetDecoratorOffsetFromCursor();
+	const FVector2f GetDecoratorOffsetFromCursor();
 
 	void BroadcastDropEvent(const FVector2D InLocation);
 
@@ -64,13 +64,13 @@ protected:
 	TSharedPtr<class SWidget> UIBeingDragged;
 
 	/** Where the user grabbed the UI measured in screen space from its top-left corner */
-	FVector2D DecoratorOffsetFromCursor;
+	FVector2f DecoratorOffsetFromCursor;
 
 	/** Decorator widget where we add temp doc tabs to */
 	TSharedPtr<SDockingTabStack> CursorDecoratorStackNode;
 
 	/** What the size of the content was when it was last shown. The user drags splitters to set this size; it is legitimate state. */
-	FVector2D LastContentSize;
+	FVector2f LastContentSize;
 
 	FOnInViewportUIDropped DropDelegate;
 };
