Move SEditorViewportDepthBar to EditorFramework, renaming it to SDepthBar. Refactors out direct references to the Editor Viewport Client.

#rnx

diff --git a/Engine/Source/Editor/EditorFramework/Private/Viewports/SDepthBar.cpp b/Engine/Source/Editor/EditorFramework/Private/Viewports/SDepthBar.cpp
new file mode 100644
--- /dev/null
+++ b/Engine/Source/Editor/EditorFramework/Private/Viewports/SDepthBar.cpp
@@ -0,0 +1,560 @@
+﻿// Copyright Epic Games, Inc. All Rights Reserved.
+
+#include "Viewports/SDepthBar.h"
+#include "Math/UnitConversion.h"
+#include "Widgets/Layout/SBorder.h"
+#include "Widgets/Text/STextBlock.h"
+#include "Layout/ArrangedChildren.h"
+
+#define LOCTEXT_NAMESPACE "DepthBar"
+
+namespace UE::EditorWidgets::Private
+{
+	constexpr float IndicatorExtraHover = 2.0f;
+	constexpr float MinIndicatorDifference = 2.0f;
+	
+	FText GetDistanceDisplayText(double Value)
+	{
+		if (FUnitConversion::Settings().ShouldDisplayUnits())
+		{
+			const TArray<EUnit>& AllDistanceUnits = FUnitConversion::Settings().GetDisplayUnits(EUnitType::Distance);
+			const EUnit BaseUnit = AllDistanceUnits.Num() ? AllDistanceUnits[0] : EUnit::Centimeters;
+			const FNumericUnit DisplayUnit = FUnitConversion::QuantizeUnitsToBestFit(Value, BaseUnit);
+			
+			const FNumberFormattingOptions Options = FNumberFormattingOptions()
+				.SetMaximumFractionalDigits(2)
+				.SetMinimumFractionalDigits(2);
+			
+			return FText::Format(
+				LOCTEXT("DistanceFormat", "{0}{1}"),
+				FText::AsNumber(DisplayUnit.Value, &Options),
+				FText::FromString(FUnitConversion::GetUnitDisplayString(DisplayUnit.Units))
+			);
+		}
+		return FText::Format(LOCTEXT("RawDistanceFormat", "{0}"), Value);
+	}
+}
+
+void SDepthBar::Construct(const FArguments& InArgs)
+{
+	Mode = InArgs._Mode;
+	
+	DepthSpace = InArgs._DepthSpace;
+	NearPlane = InArgs._NearPlane;
+	FarPlane = InArgs._FarPlane;
+	OnNearPlaneChanged = InArgs._OnNearPlaneChanged;
+	OnFarPlaneChanged = InArgs._OnFarPlaneChanged;
+
+	BackgroundBrush = FAppStyle::Get().GetBrush("EditorViewport.DepthBar.Background");
+	SliceBrush = FAppStyle::Get().GetBrush("EditorViewport.DepthBar.Slice");
+	HoveredSliceBrush = FAppStyle::Get().GetBrush("EditorViewport.DepthBar.SliceHovered");
+	HoveredSliceTopBrush = FAppStyle::Get().GetBrush("EditorViewport.DepthBar.SliceTop");
+	HoveredSliceBottomBrush = FAppStyle::Get().GetBrush("EditorViewport.DepthBar.SliceBottom");
+	
+	DepthIndicatorChildren.AddSlots(MoveTemp(const_cast<TArray<FDepthIndicatorSlot::FSlotArguments>&>(InArgs._Slots)));
+	
+	AddDepthLabelSlot()
+	.VAlign(VAlign_Top)
+	.OnGetBarPosition_Lambda([](const FBarPositions& BarPositions) { return BarPositions.BarFarPlane; })
+	[
+		SNew(SBorder)
+		.BorderImage(BackgroundBrush)
+		.Padding(6.0f, 2.0f)
+		.Visibility(this, &SDepthBar::GetFarPlaneLabelVisibility)
+		[
+			SNew(STextBlock)
+			.Text(this, &SDepthBar::GetFarPlaneText)
+		]
+	];
+	
+	AddDepthLabelSlot()
+	.VAlign(VAlign_Bottom)
+	.OnGetBarPosition_Lambda([](const FBarPositions& BarPositions) { return BarPositions.BarNearPlane; })
+	[
+		SNew(SBorder)
+		.BorderImage(BackgroundBrush)
+		.Padding(6.0f, 2.0f)
+		.Visibility(this, &SDepthBar::GetNearPlaneLabelVisibility)
+		[
+			SNew(STextBlock)
+			.Text(this, &SDepthBar::GetNearPlaneText)
+		]
+	];
+}
+
+FCursorReply SDepthBar::OnCursorQuery(const FGeometry& MyGeometry, const FPointerEvent& CursorEvent) const
+{
+	using namespace UE::EditorWidgets::Private;
+	
+	switch (DragTarget)
+	{
+		case EDragTarget::Far:
+		case EDragTarget::Near:
+			return FCursorReply::Cursor(EMouseCursor::ResizeUpDown);
+		case EDragTarget::Slice:
+			return FCursorReply::Cursor(EMouseCursor::GrabHandClosed);
+		default:
+			switch (HoverTarget)
+			{
+				case EDragTarget::Far:
+				case EDragTarget::Near:
+					return FCursorReply::Cursor(EMouseCursor::ResizeUpDown);
+				case EDragTarget::Slice:
+					return FCursorReply::Cursor(EMouseCursor::GrabHand);
+				default:
+					return FCursorReply::Unhandled();	
+			}
+	}
+}
+
+SDepthBar::FDepthIndicatorSlot& SDepthBar::GetDepthIndicatorSlot(int32 Index)
+{
+	return DepthIndicatorChildren[Index];
+}
+
+const SDepthBar::FDepthIndicatorSlot& SDepthBar::GetDepthIndicatorSlot(int32 Index) const
+{
+	return DepthIndicatorChildren[Index];
+}
+
+void SDepthBar::ClearDepthIndicators()
+{
+	DepthIndicatorChildren.Empty();
+}
+
+FReply SDepthBar::OnMouseButtonDown(const FGeometry& MyGeometry, const FPointerEvent& MouseEvent)
+{
+	using namespace UE::EditorWidgets::Private;
+
+	if (HoverTarget != EDragTarget::None)
+	{
+		if (const FBarPositions* Positions = GetBarPositions(MyGeometry).GetPtrOrNull())
+		{
+			const FVector2f MousePosition = MyGeometry.AbsoluteToLocal(MouseEvent.GetScreenSpacePosition());
+		
+			DragTarget = HoverTarget;
+			if (DragTarget == EDragTarget::Slice)
+			{
+				DragOffset = MousePosition.Y - Positions->BarFarPlane;
+			}
+			
+			return FReply::Handled().CaptureMouse(AsShared());
+		}
+	}
+
+	return FReply::Unhandled();
+}
+
+FReply SDepthBar::OnMouseMove(const FGeometry& MyGeometry, const FPointerEvent& MouseEvent)
+{
+	using namespace UE::EditorWidgets::Private;
+
+	if (DragTarget != EDragTarget::None)
+	{
+		if (const FBarPositions* Positions = GetBarPositions(MyGeometry).GetPtrOrNull())
+		{
+			const FVector2f MousePosition = MyGeometry.AbsoluteToLocal(MouseEvent.GetScreenSpacePosition());
+			
+			if (Mode.Get() == EMode::Orthographic)
+			{
+				if (DragTarget == EDragTarget::Far)
+				{
+					const float ClampedBarPosition = FMath::Clamp(MousePosition.Y, Positions->BarMin, Positions->BarNearPlane - MinIndicatorDifference);
+					
+					if (ClampedBarPosition == Positions->BarMin)
+					{
+						OnFarPlaneChanged.ExecuteIfBound(TOptional<double>());
+					}
+					else
+					{
+						OnFarPlaneChanged.ExecuteIfBound(Positions->BarToAlignedPosition(ClampedBarPosition));
+					}
+					
+					return FReply::Handled();
+				}
+				if (DragTarget == EDragTarget::Near)
+				{
+					const float ClampedBarPosition = FMath::Clamp(MousePosition.Y, Positions->BarFarPlane + MinIndicatorDifference, Positions->BarMax);
+					if (ClampedBarPosition == Positions->BarMax)
+					{
+						OnNearPlaneChanged.ExecuteIfBound(TOptional<double>());
+					}
+					else
+					{
+						OnNearPlaneChanged.ExecuteIfBound(Positions->BarToAlignedPosition(ClampedBarPosition));
+					}
+					
+					return FReply::Handled();
+				}
+				if (DragTarget == EDragTarget::Slice)
+				{
+					const float RawPosition = MousePosition.Y - DragOffset;
+					const float ClampedBarPosition = FMath::Clamp(RawPosition, Positions->BarMin, Positions->BarMax - (Positions->BarNearPlane - Positions->BarFarPlane));
+					
+					OnFarPlaneChanged.ExecuteIfBound(Positions->BarToAlignedPosition(ClampedBarPosition));
+					OnNearPlaneChanged.ExecuteIfBound(Positions->BarToAlignedPosition(ClampedBarPosition + (Positions->BarNearPlane - Positions->BarFarPlane)));
+					
+					return FReply::Handled();
+				}
+			}
+		}
+	}
+	else if (const FBarPositions* Positions = GetBarPositions(MyGeometry).GetPtrOrNull())
+	{
+		const FVector2f MousePosition = MyGeometry.AbsoluteToLocal(MouseEvent.GetScreenSpacePosition());
+		
+		const float HoverZone = SliderHeadHeight * 0.5f + IndicatorExtraHover;
+		
+		if (MousePosition.Y >= Positions->BarFarPlane - HoverZone && MousePosition.Y < Positions->BarFarPlane + HoverZone)
+		{
+			HoverTarget = EDragTarget::Far; 
+		}
+		else if (MousePosition.Y >= Positions->BarNearPlane - HoverZone && MousePosition.Y < Positions->BarNearPlane + HoverZone)
+		{
+			HoverTarget = EDragTarget::Near;
+		}
+		else if (MousePosition.Y > Positions->BarFarPlane && MousePosition.Y < Positions->BarNearPlane)
+		{
+			HoverTarget = EDragTarget::Slice;
+		}
+		else
+		{
+			HoverTarget = EDragTarget::None;
+		}
+	}
+	return FReply::Unhandled();
+}
+
+FReply SDepthBar::OnMouseButtonUp(const FGeometry& MyGeometry, const FPointerEvent& MouseEvent)
+{
+	if (DragTarget != EDragTarget::None)
+	{
+		DragTarget = EDragTarget::None;
+		HoverTarget = EDragTarget::None;
+		return FReply::Handled().ReleaseMouseCapture();
+	}
+	return FReply::Unhandled();
+}
+
+void SDepthBar::OnMouseLeave(const FPointerEvent& MouseEvent)
+{
+	SPanel::OnMouseLeave(MouseEvent);
+	
+	if (DragTarget == EDragTarget::None)
+	{
+		HoverTarget = EDragTarget::None;
+	}
+}
+
+void SDepthBar::OnArrangeChildren(const FGeometry& AllottedGeometry, FArrangedChildren& ArrangedChildren) const
+{
+	if (FBarPositions* BarPositions = GetBarPositions(AllottedGeometry).GetPtrOrNull())
+	{
+		FMargin Inset = Padding.Get();
+	
+		for (int32 ChildIndex = 0; ChildIndex < DepthIndicatorChildren.Num(); ++ChildIndex)
+		{
+			const FDepthIndicatorSlot& CurrentChild = GetDepthIndicatorSlot(ChildIndex);
+			
+			if (CurrentChild.GetBounds().IsSet())
+			{
+				FBoxSphereBounds Bounds = CurrentChild.GetBounds().Get();
+				const FSphere Sphere = Bounds.GetSphere();
+				
+				const FVector Max = Sphere.Center + BarPositions->DepthSpace.GetForward() * Sphere.W;
+				const FVector Min = Sphere.Center - BarPositions->DepthSpace.GetForward() * Sphere.W;
+					
+				const float BarMax = BarPositions->WorldToBarPosition(Max);
+				const float BarMin = BarPositions->WorldToBarPosition(Min);
+				const float BarSize = FMath::Max(BarMin - BarMax, 1.0f);
+				
+				FVector2f Offset(Inset.Left, BarMax);
+				FVector2f BoundSize(AllottedGeometry.GetLocalSize().X - Inset.Left - Inset.Right, BarSize);
+				FVector2f ChildDesiredSize = CurrentChild.GetWidget()->GetDesiredSize();
+				
+				if (ChildDesiredSize.X > BoundSize.X)
+				{
+					Offset.X -= (ChildDesiredSize.X - BoundSize.X) * 0.5f;
+					BoundSize.X = ChildDesiredSize.X;
+				}
+				if (ChildDesiredSize.Y > BoundSize.Y)
+				{
+					Offset.Y -= (ChildDesiredSize.Y - BoundSize.Y) * 0.5f;
+					BoundSize.Y = ChildDesiredSize.Y;
+				}
+				
+				ArrangedChildren.AddWidget(AllottedGeometry.MakeChild(CurrentChild.GetWidget(), BoundSize, FSlateLayoutTransform(Offset)));
+			}
+		}
+		
+		for (int32 ChildIndex = 0; ChildIndex < DepthLabelChildren.Num(); ++ChildIndex)
+		{
+			const FDepthLabelSlot& CurrentChild = DepthLabelChildren[ChildIndex];
+			
+			if (CurrentChild.GetWidget()->GetVisibility() != EVisibility::Visible)
+			{
+				continue;
+			}
+			
+			const float BarPosition = CurrentChild.GetBarPosition(*BarPositions);
+			
+			FVector2f ChildDesiredSize = CurrentChild.GetWidget()->GetDesiredSize();
+			FVector2f Offset(-ChildDesiredSize.X - 12.0f, BarPosition);
+			
+			switch (CurrentChild.GetVerticalAlignment())
+			{
+			case VAlign_Top:
+				Offset.Y -= ChildDesiredSize.Y;
+				break;
+			case VAlign_Center:
+			case VAlign_Fill:
+				Offset.Y -= ChildDesiredSize.Y * 0.5f;
+				break;
+			default:
+				break;
+			}
+			
+			ArrangedChildren.AddWidget(AllottedGeometry.MakeChild(
+				CurrentChild.GetWidget(),
+				Offset,
+				ChildDesiredSize
+			));
+		}
+	}	
+}
+
+FChildren* SDepthBar::GetChildren()
+{
+	return &AllChildren;
+}
+
+int32 SDepthBar::OnPaint(const FPaintArgs& Args, const FGeometry& AllottedGeometry, const FSlateRect& MyCullingRect, FSlateWindowElementList& OutDrawElements, int32 LayerId, const FWidgetStyle& InWidgetStyle, bool bParentEnabled) const
+{
+	using namespace UE::EditorWidgets::Private;
+
+	FSlateDrawElement::MakeBox(OutDrawElements,
+		LayerId++,
+		AllottedGeometry.ToPaintGeometry(),
+		BackgroundBrush,
+		ESlateDrawEffect::None,
+		BackgroundBrush->GetTint(InWidgetStyle)
+	);
+
+	FMargin Inset = Padding.Get();
+ 
+	if (const FBarPositions* BarPositions = GetBarPositions(AllottedGeometry).GetPtrOrNull())
+	{
+		const FSlateBrush* CurrentSliceBrush = HoverTarget == EDragTarget::Slice ? HoveredSliceBrush : SliceBrush;
+
+		if (HoverTarget == EDragTarget::Slice)
+		{
+			LayerId = SPanel::OnPaint(Args, AllottedGeometry, MyCullingRect, OutDrawElements, LayerId, InWidgetStyle, bParentEnabled) + 1;
+		}
+
+		FSlateDrawElement::MakeBox(OutDrawElements,
+			LayerId,
+			AllottedGeometry.MakeChild(
+				FVector2f(AllottedGeometry.GetLocalSize().X - Inset.Left - Inset.Right - SliderMargin * 2.0f, BarPositions->BarNearPlane - BarPositions->BarFarPlane),
+				FSlateLayoutTransform(FVector2f(Inset.Left + SliderMargin, BarPositions->BarFarPlane))
+			).ToPaintGeometry(),
+			CurrentSliceBrush,
+			ESlateDrawEffect::None,
+			CurrentSliceBrush->GetTint(InWidgetStyle).CopyWithNewOpacity(0.7f)
+		);
+		
+		if (HoverTarget != EDragTarget::Slice)
+		{
+			LayerId = SPanel::OnPaint(Args, AllottedGeometry, MyCullingRect, OutDrawElements, LayerId, InWidgetStyle, bParentEnabled) + 1;
+		}
+		
+		FSlateDrawElement::MakeBox(OutDrawElements,
+			++LayerId,
+			AllottedGeometry.MakeChild(
+				FVector2f(AllottedGeometry.GetLocalSize().X - Inset.Left - Inset.Right - SliderHeadMargin * 2.0f, SliderHeadHeight),
+				FSlateLayoutTransform(FVector2f(Inset.Left + SliderHeadMargin, BarPositions->BarFarPlane - SliderHeadHeight * 0.5f))
+			).ToPaintGeometry(),
+			HoveredSliceTopBrush,
+			ESlateDrawEffect::None,
+			HoveredSliceTopBrush->GetTint(InWidgetStyle)
+		);
+		
+		FSlateDrawElement::MakeBox(OutDrawElements,
+			++LayerId,
+			AllottedGeometry.MakeChild(
+				FVector2f(AllottedGeometry.GetLocalSize().X - Inset.Left - Inset.Right - SliderHeadMargin * 2.0f, SliderHeadHeight),
+				FSlateLayoutTransform(FVector2f(Inset.Left + SliderHeadMargin, BarPositions->BarNearPlane - SliderHeadHeight * 0.5f))
+			).ToPaintGeometry(),
+			HoveredSliceBottomBrush,
+			ESlateDrawEffect::None,
+			HoveredSliceBottomBrush->GetTint(InWidgetStyle)
+		);
+	}
+
+	return LayerId;
+}
+
+FVector2D SDepthBar::ComputeDesiredSize(float LayoutScaleMultiplier) const
+{
+	return FVector2D(20.0, 800.0);
+}
+
+TOptional<SDepthBar::FBarPositions> SDepthBar::GetBarPositions(const FGeometry& InGeometry) const
+{
+	TOptional<FDepthSpace> Space = DepthSpace.Get();
+	if (!Space.IsSet())
+	{
+		return TOptional<FBarPositions>();
+	}
+	
+	const FMargin Inset = Padding.Get();
+	
+	if (Mode.Get() == EMode::Orthographic)
+	{
+		FBarPositions Positions
+		{
+			.DepthSpace = Space.GetValue(),
+			.BarMin = Inset.Top,
+			.BarMax = InGeometry.GetLocalSize().Y - Inset.Bottom
+		};
+	
+		TOptional<double> Near = NearPlane.Get();
+		TOptional<double> Far = FarPlane.Get();
+		
+		Positions.BarFarPlane = Positions.RelativeToBarPosition(Far.IsSet() ? Space->AlignedToRelativePosition(Far.GetValue()) : 1.0);
+		Positions.BarNearPlane = Positions.RelativeToBarPosition(Near.IsSet() ? Space->AlignedToRelativePosition(Near.GetValue()) : 0.0f);
+		
+		return Positions;
+	}
+	
+	return TOptional<FBarPositions>();
+}
+
+SDepthBar::FDepthSpace::FDepthSpace(const FVector& InOrigin, const FVector& InForward, const FBox& Bounds)
+{
+	Origin = InOrigin;
+	Forward = InForward;
+	
+	const double BoundsMax = GetAlignedPosition(Bounds.Max);
+	const double BoundsMin = GetAlignedPosition(Bounds.Min);
+	
+	if (BoundsMax > BoundsMin)
+	{
+		bFlipped = false;
+		AlignedMax = BoundsMax;
+		AlignedMin = BoundsMin;
+	}
+	else
+	{
+		bFlipped = true;
+		AlignedMax = BoundsMin;
+		AlignedMin = BoundsMax;
+	}
+}
+
+double SDepthBar::FDepthSpace::GetAlignedPosition(const FVector& Position) const
+{
+	const FVector OriginToPosition = Position - Origin;
+	const FVector ProjectedPosition = OriginToPosition.ProjectOnToNormal(Forward);
+	return ProjectedPosition.Dot(Forward); 
+}
+
+double SDepthBar::FDepthSpace::AlignedToRelativePosition(double AlignedPosition) const
+{
+	const double Size = AlignedMax - AlignedMin;
+	return Size != 0.0 ? ((AlignedPosition - AlignedMin) / Size) : 0.0;
+}
+
+double SDepthBar::FDepthSpace::RelativeToAlignedPosition(double RelativePosition) const
+{
+	return FMath::Lerp(AlignedMin, AlignedMax, RelativePosition);
+}
+
+double SDepthBar::FDepthSpace::RelativeToBoundsPosition(double RelativePosition) const
+{
+	const double Adjust = Origin.ProjectOnToNormal(Forward).Dot(Forward);
+	if (bFlipped)
+	{
+		return FMath::Lerp(AlignedMax, AlignedMin, RelativePosition) + Adjust;
+	}
+	return FMath::Lerp(AlignedMin, AlignedMax, RelativePosition) + Adjust;
+}
+
+double SDepthBar::FDepthSpace::BoundsToRelativePosition(double BoundsPosition) const
+{
+	const double Adjusted = BoundsPosition - Origin.ProjectOnToNormal(Forward).Dot(Forward);
+	if (bFlipped)
+	{
+		const double Size = AlignedMin - AlignedMax;
+		return Size != 0.0 ? ((Adjusted - AlignedMax) / Size) : 0.0;
+	}
+	return AlignedToRelativePosition(Adjusted);
+}
+
+float SDepthBar::FBarPositions::WorldToBarPosition(const FVector& WorldPosition) const
+{
+	const double AxisPosition = DepthSpace.GetAlignedPosition(WorldPosition);
+	const double RelativePosition = DepthSpace.AlignedToRelativePosition(AxisPosition);
+	return RelativeToBarPosition(RelativePosition);
+}
+
+float SDepthBar::FBarPositions::RelativeToBarPosition(double RelativePosition) const
+{
+	return FMath::Clamp(1.0f - static_cast<float>(RelativePosition), 0.0f, 1.0f) * (BarMax - BarMin) + BarMin;
+}
+
+double SDepthBar::FBarPositions::BarToAlignedPosition(float BarPosition) const
+{
+	const float BarSize = BarMax - BarMin;
+	const float RelativePosition = BarSize > 0.0f ? (1.0f - (BarPosition - BarMin) / BarSize) : 0.0f;
+	return DepthSpace.RelativeToAlignedPosition(RelativePosition);
+}
+
+EVisibility SDepthBar::GetFarPlaneLabelVisibility() const
+{
+	return HoverTarget == EDragTarget::Far || HoverTarget == EDragTarget::Slice ? EVisibility::Visible : EVisibility::Collapsed;
+}
+
+EVisibility SDepthBar::GetNearPlaneLabelVisibility() const
+{
+	return HoverTarget == EDragTarget::Near || HoverTarget == EDragTarget::Slice ? EVisibility::Visible : EVisibility::Collapsed;
+}
+
+FText SDepthBar::GetFarPlaneText() const
+{
+	if (Mode.Get() == EMode::Orthographic)
+	{
+		if (const FDepthSpace* Space = DepthSpace.Get().GetPtrOrNull())
+		{
+			if (const double* Far = FarPlane.Get().GetPtrOrNull())
+			{
+				const double RelativePosition = Space->AlignedToRelativePosition(*Far);
+				const double WorldPosition = Space->RelativeToBoundsPosition(RelativePosition);
+			
+				return FText::Format(LOCTEXT("OrthoFarPlaneFormat", "Far Plane: {0}"), UE::EditorWidgets::Private::GetDistanceDisplayText(WorldPosition));
+			}
+			return LOCTEXT("OrthoFarPlaneDefault", "Infinite Far Plane");
+		}
+	}
+	return FText::GetEmpty();
+}
+
+FText SDepthBar::GetNearPlaneText() const
+{
+	if (Mode.Get() == EMode::Orthographic)
+	{
+		if (const FDepthSpace* Space = DepthSpace.Get().GetPtrOrNull())
+		{
+			if (const double* Near = NearPlane.Get().GetPtrOrNull())
+			{
+				const double RelativePosition = Space->AlignedToRelativePosition(*Near);
+				const double WorldPosition = Space->RelativeToBoundsPosition(RelativePosition);
+
+				return FText::Format(LOCTEXT("OrthoNearPlaneFormat", "Near Plane: {0}"), UE::EditorWidgets::Private::GetDistanceDisplayText(WorldPosition));
+			}
+			return LOCTEXT("OrthoNearPlaneDefault", "Infinite Near Plane");
+		}
+	}
+	return FText::GetEmpty();
+}
+
+#undef LOCTEXT_NAMESPACE
diff --git a/Engine/Source/Editor/EditorFramework/Public/Viewports/SDepthBar.h b/Engine/Source/Editor/EditorFramework/Public/Viewports/SDepthBar.h
new file mode 100644
--- /dev/null
+++ b/Engine/Source/Editor/EditorFramework/Public/Viewports/SDepthBar.h
@@ -0,0 +1,260 @@
+﻿// Copyright Epic Games, Inc. All Rights Reserved.
+
+#pragma once
+
+#include "Widgets/SPanel.h"
+#include "SlotBase.h"
+#include "Widgets/DeclarativeSyntaxSupport.h"
+#include "Layout/Children.h"
+
+/**
+ * A bar that allows for visualization and manipulation of a viewport's depth
+ */
+class SDepthBar : public SPanel
+{
+public:
+	enum class EMode
+	{
+		Perspective,
+		Orthographic
+	};
+	
+	struct FDepthSpace
+	{
+		EDITORFRAMEWORK_API FDepthSpace(const FVector& InOrigin, const FVector& InForward, const FBox& Bounds);
+		
+		const FVector& GetPosition() const { return Origin; }
+		const FVector& GetForward() const { return Forward; } 
+		
+		/** Gets the minimum value relative to originating direction */
+		double GetAlignedMin() const { return AlignedMin; }
+		/** Gets the maximum value relative to originating direction */
+		double GetAlignedMax() const { return AlignedMax; }
+		
+		/** Get the point of the position on the line starting at the origin, facing along the forward. */
+		EDITORFRAMEWORK_API double GetAlignedPosition(const FVector& Position) const;
+		
+		/** Get the relative value of the given axis-aligned position */
+		EDITORFRAMEWORK_API double AlignedToRelativePosition(double AlignedPosition) const;
+		
+		/** Get the axis-aligned value of the given relative position */
+		EDITORFRAMEWORK_API double RelativeToAlignedPosition(double RelativePosition) const;
+		
+		/** Get the bounds position of the given relative position */
+		EDITORFRAMEWORK_API double RelativeToBoundsPosition(double RelativePosition) const;
+		
+		/** Get the relative value of the given bounds position */
+		EDITORFRAMEWORK_API double BoundsToRelativePosition(double BoundsPosition) const;
+		
+	private:
+		FVector Origin;
+		FVector Forward;	
+		
+		double AlignedMin;
+		double AlignedMax;
+		
+		bool bFlipped;
+	};
+
+	/**
+	 * A slot that is placed along the depth bar, representing a range within the depth space
+	 */
+	class FDepthIndicatorSlot : public TSlotBase<FDepthIndicatorSlot>
+	{
+	public:
+	
+		SLATE_SLOT_BEGIN_ARGS(FDepthIndicatorSlot, TSlotBase<FDepthIndicatorSlot>)
+			/** Determines where in the bar the slot's contents appear. */
+			SLATE_ATTRIBUTE(FBoxSphereBounds, Bounds)
+		SLATE_SLOT_END_ARGS()
+		
+		void SetBounds(TAttribute<FBoxSphereBounds> InBounds)
+		{
+			Bounds = InBounds;
+		}
+		
+		void Construct(const FChildren& SlotOwner, FSlotArguments&& InArgs)
+		{
+			TSlotBase<FDepthIndicatorSlot>::Construct(SlotOwner, MoveTemp(InArgs));
+			Bounds = InArgs._Bounds;
+		}
+		
+		const TAttribute<FBoxSphereBounds>& GetBounds() const { return Bounds; }
+		
+	protected:
+		TAttribute<FBoxSphereBounds> Bounds;
+	};
+
+	static FDepthIndicatorSlot::FSlotArguments Slot()
+	{
+		return FDepthIndicatorSlot::FSlotArguments(MakeUnique<FDepthIndicatorSlot>());
+	}
+	
+	DECLARE_DELEGATE_OneParam(FOnDepthChanged, const TOptional<double>&)
+	
+	SLATE_BEGIN_ARGS(SDepthBar)
+	{}
+		SLATE_ATTRIBUTE(EMode, Mode)
+		SLATE_ATTRIBUTE(TOptional<FDepthSpace>, DepthSpace)
+		
+		SLATE_ATTRIBUTE(TOptional<double>, NearPlane)
+		SLATE_ATTRIBUTE(TOptional<double>, FarPlane)
+		
+		SLATE_EVENT(FOnDepthChanged, OnNearPlaneChanged)
+		SLATE_EVENT(FOnDepthChanged, OnFarPlaneChanged)
+		
+		SLATE_SLOT_ARGUMENT(SDepthBar::FDepthIndicatorSlot, Slots)
+	SLATE_END_ARGS()
+	
+	using FScopedWidgetDepthIndicatorSlotArguments = TPanelChildren<FDepthIndicatorSlot>::FScopedWidgetSlotArguments;
+	FScopedWidgetDepthIndicatorSlotArguments AddDepthIndicatorSlot()
+	{
+		return InsertDepthIndicatorSlot(INDEX_NONE);
+	}
+
+	FScopedWidgetDepthIndicatorSlotArguments InsertDepthIndicatorSlot(int32 Index = INDEX_NONE)
+	{
+		return FScopedWidgetDepthIndicatorSlotArguments(MakeUnique<FDepthIndicatorSlot>(), this->DepthIndicatorChildren, Index);
+	}
+	
+	EDITORFRAMEWORK_API FDepthIndicatorSlot& GetDepthIndicatorSlot(int32 Index);
+	EDITORFRAMEWORK_API const FDepthIndicatorSlot& GetDepthIndicatorSlot(int32 Index) const;
+	EDITORFRAMEWORK_API void ClearDepthIndicators();
+	
+	SDepthBar()
+		: DepthIndicatorChildren(this)
+		, DepthLabelChildren(this)
+		, AllChildren(this)
+		, Padding(FMargin(0.0f, 4.0f))
+		, SliderMargin(4.0f)
+		, SliderHeadMargin(2.0f)
+		, SliderHeadHeight(3.0f)
+	{
+		AllChildren.AddChildren(DepthIndicatorChildren);
+		AllChildren.AddChildren(DepthLabelChildren);
+	}
+	
+	EDITORFRAMEWORK_API void Construct(const FArguments& InArgs);
+	
+	EDITORFRAMEWORK_API virtual FCursorReply OnCursorQuery(const FGeometry& MyGeometry, const FPointerEvent& CursorEvent) const override;
+	EDITORFRAMEWORK_API virtual void OnArrangeChildren(const FGeometry& AllottedGeometry, FArrangedChildren& ArrangedChildren) const override;
+	
+	EDITORFRAMEWORK_API virtual FReply OnMouseButtonDown(const FGeometry& MyGeometry, const FPointerEvent& MouseEvent) override;
+	EDITORFRAMEWORK_API virtual FReply OnMouseMove(const FGeometry& MyGeometry, const FPointerEvent& MouseEvent) override;
+	EDITORFRAMEWORK_API virtual FReply OnMouseButtonUp(const FGeometry& MyGeometry, const FPointerEvent& MouseEvent) override;
+	EDITORFRAMEWORK_API virtual void OnMouseLeave(const FPointerEvent& MouseEvent) override;
+	
+protected:
+	EDITORFRAMEWORK_API virtual FChildren* GetChildren() override;
+	EDITORFRAMEWORK_API virtual int32 OnPaint(const FPaintArgs& Args, const FGeometry& AllottedGeometry, const FSlateRect& MyCullingRect, FSlateWindowElementList& OutDrawElements, int32 LayerId, const FWidgetStyle& InWidgetStyle, bool bParentEnabled) const override;
+	EDITORFRAMEWORK_API virtual FVector2D ComputeDesiredSize(float LayoutScaleMultiplier) const override;
+	
+private:
+	
+	struct FBarPositions
+	{
+		FDepthSpace DepthSpace;
+	
+		float BarMin;
+		float BarMax;
+		
+		float BarFarPlane;
+		float BarNearPlane;
+		
+		float WorldToBarPosition(const FVector& WorldPosition) const;
+		float RelativeToBarPosition(double RelativePosition) const;
+		double BarToAlignedPosition(float BarPosition) const;
+	};
+
+	DECLARE_DELEGATE_RetVal_OneParam(float, FGetBarPosition, const FBarPositions& BarPositions)
+
+	// Depth label slots are displayed next to the bar
+	struct FDepthLabelSlot : public TSlotBase<FDepthLabelSlot>
+	{
+		SLATE_SLOT_BEGIN_ARGS(FDepthLabelSlot, TSlotBase<FDepthLabelSlot>)
+			/** Determines where in the bar the slot's contents appear. */
+			SLATE_EVENT(FGetBarPosition, OnGetBarPosition)
+			SLATE_ATTRIBUTE(EVerticalAlignment, VAlign);
+		SLATE_SLOT_END_ARGS()
+		
+		void Construct(const FChildren& SlotOwner, FSlotArguments&& InArgs)
+		{
+			TSlotBase<FDepthLabelSlot>::Construct(SlotOwner, MoveTemp(InArgs));
+			OnGetBarPosition = InArgs._OnGetBarPosition;
+			VAlign = InArgs._VAlign;
+		}
+		
+		float GetBarPosition(const FBarPositions& BarPositions) const
+		{
+			if (OnGetBarPosition.IsBound())
+			{
+				return OnGetBarPosition.Execute(BarPositions); 
+			}
+			return 0.0f;
+		}
+		
+		EVerticalAlignment GetVerticalAlignment() const
+		{
+			if (VAlign.IsSet())
+			{
+				return VAlign.Get();
+			}
+			return VAlign_Center;
+		}
+	protected:
+		FGetBarPosition OnGetBarPosition;
+		TAttribute<EVerticalAlignment> VAlign = VAlign_Center;
+	};
+	
+	using FScopedWidgetDepthLabelSlotArguments = TPanelChildren<FDepthLabelSlot>::FScopedWidgetSlotArguments;
+	FScopedWidgetDepthLabelSlotArguments AddDepthLabelSlot()
+	{
+		return InsertDepthLabelSlot(INDEX_NONE);
+	}
+
+	FScopedWidgetDepthLabelSlotArguments InsertDepthLabelSlot(int32 Index = INDEX_NONE)
+	{
+		return FScopedWidgetDepthLabelSlotArguments(MakeUnique<FDepthLabelSlot>(), this->DepthLabelChildren, Index);
+	}
+
+	TOptional<FBarPositions> GetBarPositions(const FGeometry& InGeometry) const;
+	
+	EVisibility GetFarPlaneLabelVisibility() const;
+	EVisibility GetNearPlaneLabelVisibility() const;
+	FText GetFarPlaneText() const;
+	FText GetNearPlaneText() const;
+	
+	TPanelChildren<FDepthIndicatorSlot> DepthIndicatorChildren;
+	TPanelChildren<FDepthLabelSlot> DepthLabelChildren;
+	FCombinedChildren AllChildren;
+	
+	enum class EDragTarget
+	{
+		None,
+		Near,
+		Far,
+		Slice
+	};
+	EDragTarget HoverTarget = EDragTarget::None;
+	EDragTarget DragTarget = EDragTarget::None;
+	float DragOffset = 0.0f;
+	
+	TAttribute<FMargin> Padding;
+	TAttribute<EMode> Mode;
+	TAttribute<TOptional<FDepthSpace>> DepthSpace;
+	TAttribute<TOptional<double>> NearPlane;
+	TAttribute<TOptional<double>> FarPlane;
+	FOnDepthChanged OnNearPlaneChanged;
+	FOnDepthChanged OnFarPlaneChanged;
+	
+	float SliderMargin;
+	float SliderHeadMargin;
+	float SliderHeadHeight;
+	
+	const FSlateBrush* BackgroundBrush = nullptr;
+	const FSlateBrush* SliceBrush = nullptr;
+	const FSlateBrush* HoveredSliceBrush = nullptr;
+	const FSlateBrush* HoveredSliceTopBrush = nullptr;
+	const FSlateBrush* HoveredSliceBottomBrush = nullptr;
+};
+
diff --git a/Engine/Source/Editor/LevelEditor/Private/SLevelViewport.cpp b/Engine/Source/Editor/LevelEditor/Private/SLevelViewport.cpp
--- a/Engine/Source/Editor/LevelEditor/Private/SLevelViewport.cpp
+++ b/Engine/Source/Editor/LevelEditor/Private/SLevelViewport.cpp
@@ -103,7 +103,7 @@
 #include "LevelViewportLayout.h"
 #include "EditorViewportTabContent.h"
 #include "EditorViewportCommands.h"
-#include "SEditorViewportDepthBar.h"
+#include "Viewports/SDepthBar.h"
 #include "FunctionalUIScreenshotTest.h"
 #include "GPUSkinCache.h"
 #include "GPUSkinCacheVisualizationMenuCommands.h"
diff --git a/Engine/Source/Editor/UnrealEd/Private/SEditorViewport.cpp b/Engine/Source/Editor/UnrealEd/Private/SEditorViewport.cpp
--- a/Engine/Source/Editor/UnrealEd/Private/SEditorViewport.cpp
+++ b/Engine/Source/Editor/UnrealEd/Private/SEditorViewport.cpp
@@ -36,7 +36,7 @@
 	#include "RenderGraph.h"
 #endif
 #include "LevelEditorViewport.h"
-#include "SEditorViewportDepthBar.h"
+#include "Viewports/SDepthBar.h"
 #include "SEditorViewportGridPanel.h"
 #include "ViewportInteractions/ViewportInteractionsBehaviorSource.h"
 #include "Widgets/Input/SButton.h"
@@ -304,7 +304,13 @@ void SEditorViewport::Construct( const FArguments& InArgs )
 			+ SVerticalBox::Slot()
 			.FillHeight(1.0f)
 			[
-				SAssignNew(DepthBar, SEditorViewportDepthBar, GetViewportClient())
+				SAssignNew(DepthBar, SDepthBar)
+				.Mode(this, &SEditorViewport::GetDepthBarMode)
+				.DepthSpace(this, &SEditorViewport::GetDepthBarSpace)
+				.NearPlane(this, &SEditorViewport::GetDepthBarNearPlane)
+				.FarPlane(this, &SEditorViewport::GetDepthBarFarPlane)
+				.OnNearPlaneChanged(this, &SEditorViewport::OnDepthBarNearPlaneChanged)
+				.OnFarPlaneChanged(this, &SEditorViewport::OnDepthBarFarPlaneChanged)
 			]
 			+ SVerticalBox::Slot()
 			.Padding(0.0f, 6.0f, 0.0f, 0.0f)
@@ -756,6 +762,60 @@ EVisibility SEditorViewport::OnGetDepthBarVisibility() const
 	return EVisibility::Collapsed;
 }
 
+SDepthBar::EMode SEditorViewport::GetDepthBarMode() const
+{
+	return Client->IsOrtho() ? SDepthBar::EMode::Orthographic : SDepthBar::EMode::Perspective;
+}
+
+TOptional<SDepthBar::FDepthSpace> SEditorViewport::GetDepthBarSpace() const
+{
+	FBox WorldBounds = Client->GetWorldBounds();
+	if (!WorldBounds.IsValid)
+	{
+		return TOptional<SDepthBar::FDepthSpace>();
+	}
+	
+	return SDepthBar::FDepthSpace(Client->GetViewLocation(), Client->GetForwardVector(), WorldBounds);
+}
+
+TOptional<double> SEditorViewport::GetDepthBarNearPlane() const
+{
+	if (Client->IsOrtho())
+	{
+		return Client->GetOrthographicNearPlaneOverride();
+	}
+	return TOptional<double>();
+}
+
+TOptional<double> SEditorViewport::GetDepthBarFarPlane() const
+{
+	if (Client->IsOrtho())
+    {
+    	return Client->GetOrthographicFarPlaneOverride();
+    }
+    return TOptional<double>();
+}
+
+void SEditorViewport::OnDepthBarNearPlaneChanged(const TOptional<double>& InNearPlane)
+{
+	if (Client->IsOrtho())
+	{
+		Client->SetOrthographicNearPlaneOverride(InNearPlane);
+	}
+	
+	Client->Invalidate();
+}
+
+void SEditorViewport::OnDepthBarFarPlaneChanged(const TOptional<double>& InFarPlane)
+{
+	if (Client->IsOrtho())
+	{
+		Client->SetOrthographicFarPlaneOverride(InFarPlane);
+	}
+	
+	Client->Invalidate();
+}
+
 FReply SEditorViewport::DismissDepthBar()
 {
 	Client->SetOrthographicFarPlaneOverride(TOptional<double>());
diff --git a/Engine/Source/Editor/UnrealEd/Private/SEditorViewportDepthBar.cpp b/Engine/Source/Editor/UnrealEd/Private/SEditorViewportDepthBar.cpp
deleted file mode 100644
diff --git a/Engine/Source/Editor/UnrealEd/Private/ViewportToolbar/UnrealEdViewportToolbar.cpp b/Engine/Source/Editor/UnrealEd/Private/ViewportToolbar/UnrealEdViewportToolbar.cpp
--- a/Engine/Source/Editor/UnrealEd/Private/ViewportToolbar/UnrealEdViewportToolbar.cpp
+++ b/Engine/Source/Editor/UnrealEd/Private/ViewportToolbar/UnrealEdViewportToolbar.cpp
@@ -20,7 +20,6 @@
 #include "LevelEditorActions.h"
 #include "RayTracingDebugVisualizationMenuCommands.h"
 #include "SEditorViewport.h"
-#include "SEditorViewportDepthBar.h"
 #include "SScalabilitySettings.h"
 #include "ShowFlagMenuCommands.h"
 #include "Settings/EditorProjectSettings.h"
@@ -821,7 +820,7 @@ public:
 			const FBox Bounds = Client->GetWorldBounds();
 			if (Bounds.IsValid)
 			{
-				const SEditorViewportDepthBar::FDepthSpace DepthSpace = GetDepthSpace();
+				const SDepthBar::FDepthSpace DepthSpace = GetDepthSpace();
 				const double BoundsMin = DepthSpace.RelativeToBoundsPosition(0.0f);
 				const double BoundsMax = DepthSpace.RelativeToBoundsPosition(1.0f);
 				
@@ -865,7 +864,7 @@ private:
 		{
 			// The overrides are "aligned" offsets relative to the camera position
 			// we want to present them as world coordinates.
-			const SEditorViewportDepthBar::FDepthSpace DepthSpace = GetDepthSpace();
+			const SDepthBar::FDepthSpace DepthSpace = GetDepthSpace();
 			const double RelativePosition = DepthSpace.AlignedToRelativePosition(*Value);
 			return static_cast<float>(DepthSpace.RelativeToBoundsPosition(RelativePosition));
 		}
@@ -876,7 +875,7 @@ private:
 	{
 		if (TSharedPtr<FEditorViewportClient> Client = WeakClient.Pin())
 		{
-			const SEditorViewportDepthBar::FDepthSpace DepthSpace = GetDepthSpace();
+			const SDepthBar::FDepthSpace DepthSpace = GetDepthSpace();
 			const double RelativePosition = DepthSpace.BoundsToRelativePosition(NewValue);
 			
 			if (Mode == EMode::Near && FMath::IsNearlyEqual(RelativePosition, 0.0))
@@ -919,11 +918,11 @@ private:
 		}
 	}
 	
-	SEditorViewportDepthBar::FDepthSpace GetDepthSpace() const
+	SDepthBar::FDepthSpace GetDepthSpace() const
 	{
 		TSharedPtr<FEditorViewportClient> Client = WeakClient.Pin();
 		
-		return SEditorViewportDepthBar::FDepthSpace(Client->GetViewLocation(), Client->GetForwardVector(), Client->GetWorldBounds());
+		return SDepthBar::FDepthSpace(Client->GetViewLocation(), Client->GetForwardVector(), Client->GetWorldBounds());
 	}
 	
 	TOptional<FText> GetDisplayValue(float Value) const
diff --git a/Engine/Source/Editor/UnrealEd/Public/SEditorViewport.h b/Engine/Source/Editor/UnrealEd/Public/SEditorViewport.h
--- a/Engine/Source/Editor/UnrealEd/Public/SEditorViewport.h
+++ b/Engine/Source/Editor/UnrealEd/Public/SEditorViewport.h
@@ -10,6 +10,7 @@
 #include "Widgets/SViewport.h"
 #include "UnrealWidgetFwd.h"
 #include "EditorViewportClient.h"
+#include "Viewports/SDepthBar.h"
 
 class FActiveTimerHandle;
 class FSceneViewport;
@@ -20,7 +21,7 @@ class UModeManagerInteractiveToolsContext;
 class UViewportInteractionsBehaviorSource;
 struct FSlateBrush;
 struct FToolMenuContext;
-class SEditorViewportDepthBar;
+class SDepthBar;
 
 class SEditorViewport
 	: public SCompoundWidget
@@ -355,7 +356,7 @@ protected:
 	TSharedPtr<FEditorViewportClient> Client;
 	
 	/** The depth bar used by this viewport */
-	TSharedPtr<SEditorViewportDepthBar> DepthBar;
+	TSharedPtr<SDepthBar> DepthBar;
 	
 	/** Commandlist used in the viewport (Maps commands to viewport specific actions) */
 	TSharedPtr<FUICommandList> CommandList;
@@ -383,6 +384,24 @@ private:
 	/** @return The visibility of the viewport depth bar. */
 	EVisibility OnGetDepthBarVisibility() const;
 	
+	/** @return The current mode of the depth bar. */
+	SDepthBar::EMode GetDepthBarMode() const;
+	
+	/** @return The current depth space of the depth bar. */
+	TOptional<SDepthBar::FDepthSpace> GetDepthBarSpace() const;
+	
+	/** @return The current near plane for the depth bar. */
+	TOptional<double> GetDepthBarNearPlane() const;
+	
+	/** @return The current near plane for the depth bar. */
+	TOptional<double> GetDepthBarFarPlane() const;
+	
+	/** Called to change the depth bar near plane. */
+	void OnDepthBarNearPlaneChanged(const TOptional<double>& InNearPlane);
+	
+	/** Called to change the depth bar far plane. */
+	void OnDepthBarFarPlaneChanged(const TOptional<double>& InFarPlane);
+	
 	/** Reset the depth bar back to the default state. */
 	FReply DismissDepthBar();
 
diff --git a/Engine/Source/Editor/UnrealEd/Public/SEditorViewportDepthBar.h b/Engine/Source/Editor/UnrealEd/Public/SEditorViewportDepthBar.h
deleted file mode 100644
