Transition viewport depth bar to be part of a larger widget that includes a vertical toolbar.

[RN] minor editor

diff --git a/Engine/Source/Editor/EditorFramework/Private/Styling/DepthBarStyle.cpp b/Engine/Source/Editor/EditorFramework/Private/Styling/DepthBarStyle.cpp
new file mode 100644
--- /dev/null
+++ b/Engine/Source/Editor/EditorFramework/Private/Styling/DepthBarStyle.cpp
@@ -0,0 +1,45 @@
+﻿// Copyright Epic Games, Inc. All Rights Reserved.
+
+
+#include "Styling/DepthBarStyle.h"
+
+#include "Brushes/SlateNoResource.h"
+
+const FName FDepthBarStyle::TypeName("FDepthBarStyle");
+
+FDepthBarStyle::FDepthBarStyle()
+	: BackgroundBrush(FSlateNoResource())
+	, TrackBrush(FSlateNoResource())
+	, SliceNormalBrush(FSlateNoResource())
+	, SliceHoveredBrush(FSlateNoResource())
+	, SliceTopBrush(FSlateNoResource())
+	, SliceBottomBrush(FSlateNoResource())
+	, SliceBottomHoveredBrush(FSlateNoResource())
+	, FarPlaneButtonBrush(FSlateNoResource())
+	, NearPlaneButtonBrush(FSlateNoResource())
+{
+}
+
+FDepthBarStyle::FDepthBarStyle(const FDepthBarStyle&) = default;
+
+FDepthBarStyle::~FDepthBarStyle() = default;
+
+const FDepthBarStyle& FDepthBarStyle::GetDefault()
+{
+	static FDepthBarStyle Default;
+	return Default;
+}
+
+void FDepthBarStyle::GetResources(TArray<const FSlateBrush*>& OutBrushes) const
+{
+	OutBrushes.Add(&BackgroundBrush);
+	OutBrushes.Add(&TrackBrush);
+	OutBrushes.Add(&SliceNormalBrush);
+	OutBrushes.Add(&SliceHoveredBrush);
+	OutBrushes.Add(&SliceTopBrush);
+	OutBrushes.Add(&SliceTopHoveredBrush);
+	OutBrushes.Add(&SliceBottomBrush);
+	OutBrushes.Add(&SliceBottomHoveredBrush);
+	OutBrushes.Add(&FarPlaneButtonBrush);
+	OutBrushes.Add(&NearPlaneButtonBrush);
+}
diff --git a/Engine/Source/Editor/EditorFramework/Private/Viewports/SDepthBar.cpp b/Engine/Source/Editor/EditorFramework/Private/Viewports/SDepthBar.cpp
--- a/Engine/Source/Editor/EditorFramework/Private/Viewports/SDepthBar.cpp
+++ b/Engine/Source/Editor/EditorFramework/Private/Viewports/SDepthBar.cpp
@@ -10,7 +10,6 @@
 
 namespace UE::EditorWidgets::Private
 {
-	constexpr float IndicatorExtraHover = 2.0f;
 	constexpr float MinIndicatorDifference = 2.0f;
 	
 	FText GetDistanceDisplayText(double Value)
@@ -45,12 +44,15 @@ void SDepthBar::Construct(const FArguments& InArgs)
 	OnNearPlaneChanged = InArgs._OnNearPlaneChanged;
 	OnFarPlaneChanged = InArgs._OnFarPlaneChanged;
 
-	BackgroundBrush = FAppStyle::Get().GetBrush("EditorViewport.DepthBar.Background");
-	SliceBrush = FAppStyle::Get().GetBrush("EditorViewport.DepthBar.Slice");
-	HoveredSliceBrush = FAppStyle::Get().GetBrush("EditorViewport.DepthBar.SliceHovered");
-	HoveredSliceTopBrush = FAppStyle::Get().GetBrush("EditorViewport.DepthBar.SliceTop");
-	HoveredSliceBottomBrush = FAppStyle::Get().GetBrush("EditorViewport.DepthBar.SliceBottom");
-	
+	if (InArgs._Style)
+	{
+		Style = InArgs._Style;
+	}
+	else
+	{
+		Style = &FDepthBarStyle::GetDefault();
+	}
+
 	DepthIndicatorChildren.AddSlots(MoveTemp(const_cast<TArray<FDepthIndicatorSlot::FSlotArguments>&>(InArgs._Slots)));
 	
 	AddDepthLabelSlot()
@@ -58,7 +60,7 @@ void SDepthBar::Construct(const FArguments& InArgs)
 	.OnGetBarPosition_Lambda([](const FBarPositions& BarPositions) { return BarPositions.BarFarPlane; })
 	[
 		SNew(SBorder)
-		.BorderImage(BackgroundBrush)
+		.BorderImage(&Style->BackgroundBrush)
 		.Padding(6.0f, 2.0f)
 		.Visibility(this, &SDepthBar::GetFarPlaneLabelVisibility)
 		[
@@ -72,7 +74,7 @@ void SDepthBar::Construct(const FArguments& InArgs)
 	.OnGetBarPosition_Lambda([](const FBarPositions& BarPositions) { return BarPositions.BarNearPlane; })
 	[
 		SNew(SBorder)
-		.BorderImage(BackgroundBrush)
+		.BorderImage(&Style->BackgroundBrush)
 		.Padding(6.0f, 2.0f)
 		.Visibility(this, &SDepthBar::GetNearPlaneLabelVisibility)
 		[
@@ -88,21 +90,21 @@ FCursorReply SDepthBar::OnCursorQuery(const FGeometry& MyGeometry, const FPointe
 	
 	switch (DragTarget)
 	{
-		case EDragTarget::Far:
-		case EDragTarget::Near:
+		case EDragTarget::Far_Handle:
+		case EDragTarget::Near_Handle:
 			return FCursorReply::Cursor(EMouseCursor::ResizeUpDown);
 		case EDragTarget::Slice:
 			return FCursorReply::Cursor(EMouseCursor::GrabHandClosed);
 		default:
 			switch (HoverTarget)
 			{
-				case EDragTarget::Far:
-				case EDragTarget::Near:
+				case EDragTarget::Far_Handle:
+				case EDragTarget::Near_Handle:
 					return FCursorReply::Cursor(EMouseCursor::ResizeUpDown);
 				case EDragTarget::Slice:
 					return FCursorReply::Cursor(EMouseCursor::GrabHand);
 				default:
-					return FCursorReply::Unhandled();	
+					return FCursorReply::Cursor(EMouseCursor::Default);	
 			}
 	}
 }
@@ -128,7 +130,8 @@ FReply SDepthBar::OnMouseButtonDown(const FGeometry& MyGeometry, const FPointerE
 
 	if (HoverTarget != EDragTarget::None)
 	{
-		if (const FBarPositions* Positions = GetBarPositions(MyGeometry).GetPtrOrNull())
+		TOptional<FBarPositions> OptionalPositions = GetBarPositions(MyGeometry); 
+		if (const FBarPositions* Positions = OptionalPositions.GetPtrOrNull())
 		{
 			const FVector2f MousePosition = MyGeometry.AbsoluteToLocal(MouseEvent.GetScreenSpacePosition());
 		
@@ -149,15 +152,17 @@ FReply SDepthBar::OnMouseMove(const FGeometry& MyGeometry, const FPointerEvent&
 {
 	using namespace UE::EditorWidgets::Private;
 
+	TOptional<FBarPositions> OptionalPositions = GetBarPositions(MyGeometry);
+
 	if (DragTarget != EDragTarget::None)
 	{
-		if (const FBarPositions* Positions = GetBarPositions(MyGeometry).GetPtrOrNull())
+		if (const FBarPositions* Positions = OptionalPositions.GetPtrOrNull())
 		{
 			const FVector2f MousePosition = MyGeometry.AbsoluteToLocal(MouseEvent.GetScreenSpacePosition());
 			
 			if (Mode.Get() == EMode::Orthographic)
 			{
-				if (DragTarget == EDragTarget::Far)
+				if (DragTarget == EDragTarget::Far_Handle)
 				{
 					const float ClampedBarPosition = FMath::Clamp(MousePosition.Y, Positions->BarMin, Positions->BarNearPlane - MinIndicatorDifference);
 					
@@ -172,7 +177,7 @@ FReply SDepthBar::OnMouseMove(const FGeometry& MyGeometry, const FPointerEvent&
 					
 					return FReply::Handled();
 				}
-				if (DragTarget == EDragTarget::Near)
+				if (DragTarget == EDragTarget::Near_Handle)
 				{
 					const float ClampedBarPosition = FMath::Clamp(MousePosition.Y, Positions->BarFarPlane + MinIndicatorDifference, Positions->BarMax);
 					if (ClampedBarPosition == Positions->BarMax)
@@ -199,24 +204,34 @@ FReply SDepthBar::OnMouseMove(const FGeometry& MyGeometry, const FPointerEvent&
 			}
 		}
 	}
-	else if (const FBarPositions* Positions = GetBarPositions(MyGeometry).GetPtrOrNull())
+	else if (const FBarPositions* Positions = OptionalPositions.GetPtrOrNull())
 	{
 		const FVector2f MousePosition = MyGeometry.AbsoluteToLocal(MouseEvent.GetScreenSpacePosition());
 		
-		const float HoverZone = SliderHeadHeight * 0.5f + IndicatorExtraHover;
+		const float SliderHeight = Style->SliceHandleSize.Y;
+		
+		const FMargin Inset = GetBarInset();
 		
-		if (MousePosition.Y >= Positions->BarFarPlane - HoverZone && MousePosition.Y < Positions->BarFarPlane + HoverZone)
+		if (MousePosition.Y > Style->Padding.Top && MousePosition.Y < Style->Padding.Top + Style->FarPlaneButtonPadding.GetTotalSpaceAlong<Orient_Vertical>() + Style->PlaneButtonSize.Y)
+		{
+			HoverTarget = EDragTarget::Far_Button;
+		}
+		else if (MousePosition.Y >= Positions->BarFarPlane - SliderHeight && MousePosition.Y < Positions->BarFarPlane)
 		{
-			HoverTarget = EDragTarget::Far; 
+			HoverTarget = EDragTarget::Far_Handle; 
 		}
-		else if (MousePosition.Y >= Positions->BarNearPlane - HoverZone && MousePosition.Y < Positions->BarNearPlane + HoverZone)
+		else if (MousePosition.Y >= Positions->BarNearPlane && MousePosition.Y < Positions->BarNearPlane + SliderHeight)
 		{
-			HoverTarget = EDragTarget::Near;
+			HoverTarget = EDragTarget::Near_Handle;
 		}
 		else if (MousePosition.Y > Positions->BarFarPlane && MousePosition.Y < Positions->BarNearPlane)
 		{
 			HoverTarget = EDragTarget::Slice;
 		}
+		else if (MousePosition.Y > MyGeometry.GetLocalSize().Y - Inset.Bottom)
+		{
+			HoverTarget = EDragTarget::Near_Button;
+		}
 		else
 		{
 			HoverTarget = EDragTarget::None;
@@ -236,6 +251,21 @@ FReply SDepthBar::OnMouseButtonUp(const FGeometry& MyGeometry, const FPointerEve
 	return FReply::Unhandled();
 }
 
+FReply SDepthBar::OnMouseButtonDoubleClick(const FGeometry& InMyGeometry, const FPointerEvent& InMouseEvent)
+{
+	if (HoverTarget == EDragTarget::Far_Button)
+	{
+		OnFarPlaneChanged.ExecuteIfBound(TOptional<double>());
+		return FReply::Handled().ReleaseMouseCapture();
+	}
+	if (HoverTarget == EDragTarget::Near_Button)
+	{
+		OnNearPlaneChanged.ExecuteIfBound(TOptional<double>());
+		return FReply::Handled().ReleaseMouseCapture();
+	}
+	return FReply::Unhandled();
+}
+
 void SDepthBar::OnMouseLeave(const FPointerEvent& MouseEvent)
 {
 	SPanel::OnMouseLeave(MouseEvent);
@@ -248,10 +278,12 @@ void SDepthBar::OnMouseLeave(const FPointerEvent& MouseEvent)
 
 void SDepthBar::OnArrangeChildren(const FGeometry& AllottedGeometry, FArrangedChildren& ArrangedChildren) const
 {
-	if (FBarPositions* BarPositions = GetBarPositions(AllottedGeometry).GetPtrOrNull())
+	TOptional<FBarPositions> OptionalPositions = GetBarPositions(AllottedGeometry); 
+	if (FBarPositions* BarPositions = OptionalPositions.GetPtrOrNull())
 	{
-		FMargin Inset = Padding.Get();
-	
+		FMargin Inset = GetBarInset();
+		const FVector2f Center = AllottedGeometry.GetLocalSize() * 0.5f;
+		
 		for (int32 ChildIndex = 0; ChildIndex < DepthIndicatorChildren.Num(); ++ChildIndex)
 		{
 			const FDepthIndicatorSlot& CurrentChild = GetDepthIndicatorSlot(ChildIndex);
@@ -268,20 +300,63 @@ void SDepthBar::OnArrangeChildren(const FGeometry& AllottedGeometry, FArrangedCh
 				const float BarMin = BarPositions->WorldToBarPosition(Min);
 				const float BarSize = FMath::Max(BarMin - BarMax, 1.0f);
 				
-				FVector2f Offset(Inset.Left, BarMax);
-				FVector2f BoundSize(AllottedGeometry.GetLocalSize().X - Inset.Left - Inset.Right, BarSize);
+				FVector2f Offset(Inset.Left, 0.0f);
+				FVector2f BoundSize(AllottedGeometry.GetLocalSize().X - Inset.GetTotalSpaceAlong<Orient_Horizontal>(), BarSize);
+				
 				FVector2f ChildDesiredSize = CurrentChild.GetWidget()->GetDesiredSize();
 				
+				const EHorizontalAlignment HorizontalAlignment = CurrentChild.GetHorizontalAlignment();
+				switch (HorizontalAlignment)
+				{
+				case HAlign_Left:
+				case HAlign_Fill:
+					Offset.X = Inset.Left;
+					break;
+				case HAlign_Right:
+					Offset.X = Inset.Right - ChildDesiredSize.X;
+					break;
+				case HAlign_Center:
+					Offset.X = Center.X - ChildDesiredSize.X * 0.5f;
+					break;
+				}
+				
+				const EVerticalAlignment VerticalAlignment = CurrentChild.GetVerticalAlignment(); 
+				switch (VerticalAlignment)
+				{
+				case VAlign_Top:
+					Offset.Y = BarMax - ChildDesiredSize.Y;
+					break;
+				case VAlign_Bottom:
+					Offset.Y = BarMin;
+					break;
+				case VAlign_Center:
+					Offset.Y = BarMax + (BarMin - BarMax) * 0.5f - ChildDesiredSize.Y * 0.5f;
+					break;
+				case VAlign_Fill:
+					Offset.Y = BarMax;
+					break;
+				}
+				
+				// Allow minimum child size to overstep the bounds
 				if (ChildDesiredSize.X > BoundSize.X)
 				{
 					Offset.X -= (ChildDesiredSize.X - BoundSize.X) * 0.5f;
 					BoundSize.X = ChildDesiredSize.X;
 				}
+				else if (HorizontalAlignment != HAlign_Fill)
+				{
+					BoundSize.X = ChildDesiredSize.X;
+				}
+				
 				if (ChildDesiredSize.Y > BoundSize.Y)
 				{
 					Offset.Y -= (ChildDesiredSize.Y - BoundSize.Y) * 0.5f;
 					BoundSize.Y = ChildDesiredSize.Y;
 				}
+				else if (VerticalAlignment != VAlign_Fill)
+				{
+					BoundSize.Y = ChildDesiredSize.Y;
+				}
 				
 				ArrangedChildren.AddWidget(AllottedGeometry.MakeChild(CurrentChild.GetWidget(), BoundSize, FSlateLayoutTransform(Offset)));
 			}
@@ -332,34 +407,91 @@ int32 SDepthBar::OnPaint(const FPaintArgs& Args, const FGeometry& AllottedGeomet
 {
 	using namespace UE::EditorWidgets::Private;
 
+	const FVector2f Center = AllottedGeometry.GetLocalSize() * 0.5f;
+
+	// Background
 	FSlateDrawElement::MakeBox(OutDrawElements,
 		LayerId++,
 		AllottedGeometry.ToPaintGeometry(),
-		BackgroundBrush,
+		&Style->BackgroundBrush,
 		ESlateDrawEffect::None,
-		BackgroundBrush->GetTint(InWidgetStyle)
+		Style->BackgroundBrush.GetTint(InWidgetStyle)
 	);
-
-	FMargin Inset = Padding.Get();
- 
-	if (const FBarPositions* BarPositions = GetBarPositions(AllottedGeometry).GetPtrOrNull())
+	
+	// Plane Buttons
 	{
-		const FSlateBrush* CurrentSliceBrush = HoverTarget == EDragTarget::Slice ? HoveredSliceBrush : SliceBrush;
+		const float PlaneButtonHorizontalInset = Center.X - Style->PlaneButtonSize.X * 0.5f;  
+	
+		FSlateDrawElement::MakeBox(OutDrawElements,
+			LayerId,
+			AllottedGeometry.MakeChild(
+				FVector2f(Style->PlaneButtonSize),
+				FSlateLayoutTransform(FVector2f(
+					PlaneButtonHorizontalInset,
+					Style->Padding.Top + Style->FarPlaneButtonPadding.Top
+				))
+			).ToPaintGeometry(),
+			&Style->FarPlaneButtonBrush,
+			ESlateDrawEffect::None,
+			(HoverTarget == EDragTarget::Far_Button ? Style->PlaneButtonHoveredColor : Style->PlaneButtonNormalColor).GetColor(InWidgetStyle)
+		);
+	
+		FSlateDrawElement::MakeBox(OutDrawElements,
+			LayerId,
+			AllottedGeometry.MakeChild(
+				FVector2f(Style->PlaneButtonSize),
+				FSlateLayoutTransform(FVector2f(
+					PlaneButtonHorizontalInset,
+					AllottedGeometry.GetLocalSize().Y - Style->PlaneButtonSize.Y - Style->Padding.Bottom - Style->NearPlaneButtonPadding.Bottom
+				))
+			).ToPaintGeometry(),
+			&Style->NearPlaneButtonBrush,
+			ESlateDrawEffect::None,
+			(HoverTarget == EDragTarget::Near_Button ? Style->PlaneButtonHoveredColor : Style->PlaneButtonNormalColor).GetColor(InWidgetStyle)
+		);
+	}
+	
+	// Bar slider
+	TOptional<FBarPositions> OptionalPositions = GetBarPositions(AllottedGeometry); 
+	if (const FBarPositions* BarPositions = OptionalPositions.GetPtrOrNull())
+	{
+		// Draw the track
+		FSlateDrawElement::MakeBox(OutDrawElements,
+			LayerId++,
+			AllottedGeometry.MakeChild(
+				FVector2f(Style->TrackWidth, BarPositions->BarMax - BarPositions->BarMin + Style->SliceHandleSize.Y * 2.0f),
+				FSlateLayoutTransform(FVector2f(
+					Center.X - Style->TrackWidth * 0.5f,
+					BarPositions->BarMin - Style->SliceHandleSize.Y
+				))
+			).ToPaintGeometry(),
+			&Style->TrackBrush,
+			ESlateDrawEffect::None,
+			Style->TrackBrush.GetTint(InWidgetStyle)
+		);
 
 		if (HoverTarget == EDragTarget::Slice)
 		{
 			LayerId = SPanel::OnPaint(Args, AllottedGeometry, MyCullingRect, OutDrawElements, LayerId, InWidgetStyle, bParentEnabled) + 1;
 		}
+	
+		const FSlateBrush* CurrentSliceBrush = HoverTarget == EDragTarget::Slice ? &Style->SliceHoveredBrush : &Style->SliceNormalBrush;
+		FLinearColor SliceTint = CurrentSliceBrush->GetTint(InWidgetStyle);
+		if (HoverTarget == EDragTarget::Slice)
+		{
+			SliceTint.A = 0.5f;
+		}
 
+		// Draw the slice
 		FSlateDrawElement::MakeBox(OutDrawElements,
 			LayerId,
 			AllottedGeometry.MakeChild(
-				FVector2f(AllottedGeometry.GetLocalSize().X - Inset.Left - Inset.Right - SliderMargin * 2.0f, BarPositions->BarNearPlane - BarPositions->BarFarPlane),
-				FSlateLayoutTransform(FVector2f(Inset.Left + SliderMargin, BarPositions->BarFarPlane))
+				FVector2f(Style->SliceWidth, BarPositions->BarNearPlane - BarPositions->BarFarPlane),
+				FSlateLayoutTransform(FVector2f(Center.X - Style->SliceWidth * 0.5f, BarPositions->BarFarPlane))
 			).ToPaintGeometry(),
 			CurrentSliceBrush,
 			ESlateDrawEffect::None,
-			CurrentSliceBrush->GetTint(InWidgetStyle).CopyWithNewOpacity(0.7f)
+			SliceTint
 		);
 		
 		if (HoverTarget != EDragTarget::Slice)
@@ -367,27 +499,59 @@ int32 SDepthBar::OnPaint(const FPaintArgs& Args, const FGeometry& AllottedGeomet
 			LayerId = SPanel::OnPaint(Args, AllottedGeometry, MyCullingRect, OutDrawElements, LayerId, InWidgetStyle, bParentEnabled) + 1;
 		}
 		
+		// Draw the slice handles
+		const FSlateBrush* TopBrush = HoverTarget == EDragTarget::Far_Handle ? &Style->SliceTopHoveredBrush : &Style->SliceTopBrush;
 		FSlateDrawElement::MakeBox(OutDrawElements,
 			++LayerId,
 			AllottedGeometry.MakeChild(
-				FVector2f(AllottedGeometry.GetLocalSize().X - Inset.Left - Inset.Right - SliderHeadMargin * 2.0f, SliderHeadHeight),
-				FSlateLayoutTransform(FVector2f(Inset.Left + SliderHeadMargin, BarPositions->BarFarPlane - SliderHeadHeight * 0.5f))
+				FVector2f(Style->SliceHandleSize),
+				FSlateLayoutTransform(FVector2f(Center.X - Style->SliceHandleSize.X * 0.5f, BarPositions->BarFarPlane - Style->SliceHandleSize.Y))
 			).ToPaintGeometry(),
-			HoveredSliceTopBrush,
+			TopBrush,
 			ESlateDrawEffect::None,
-			HoveredSliceTopBrush->GetTint(InWidgetStyle)
+			TopBrush->GetTint(InWidgetStyle)
 		);
 		
+		const FSlateBrush* BottomBrush = HoverTarget == EDragTarget::Near_Handle ? &Style->SliceBottomHoveredBrush : &Style->SliceBottomBrush;
 		FSlateDrawElement::MakeBox(OutDrawElements,
 			++LayerId,
 			AllottedGeometry.MakeChild(
-				FVector2f(AllottedGeometry.GetLocalSize().X - Inset.Left - Inset.Right - SliderHeadMargin * 2.0f, SliderHeadHeight),
-				FSlateLayoutTransform(FVector2f(Inset.Left + SliderHeadMargin, BarPositions->BarNearPlane - SliderHeadHeight * 0.5f))
+				FVector2f(Style->SliceHandleSize),
+				FSlateLayoutTransform(FVector2f(Center.X - Style->SliceHandleSize.X * 0.5f, BarPositions->BarNearPlane))
 			).ToPaintGeometry(),
-			HoveredSliceBottomBrush,
+			BottomBrush,
 			ESlateDrawEffect::None,
-			HoveredSliceBottomBrush->GetTint(InWidgetStyle)
+			BottomBrush->GetTint(InWidgetStyle)
 		);
+		
+		// Draw the "shades" over out-of-slice content
+		if (BarPositions->BarMin < BarPositions->BarFarPlane - Style->SliceHandleSize.Y)
+		{
+			FSlateDrawElement::MakeBox(OutDrawElements,
+				LayerId,
+				AllottedGeometry.MakeChild(
+					FVector2f(AllottedGeometry.GetLocalSize().X, BarPositions->BarFarPlane - Style->SliceHandleSize.Y - BarPositions->BarMin),
+					FSlateLayoutTransform(FVector2f(0.0f, BarPositions->BarMin)
+				)).ToPaintGeometry(),
+				&Style->BackgroundBrush,
+				ESlateDrawEffect::None,
+				Style->BackgroundBrush.GetTint(InWidgetStyle).CopyWithNewOpacity(0.5f)
+			);
+		}
+		
+		if (BarPositions->BarMax > BarPositions->BarNearPlane + Style->SliceHandleSize.Y)
+		{
+			FSlateDrawElement::MakeBox(OutDrawElements,
+				LayerId,
+				AllottedGeometry.MakeChild(
+					FVector2f(AllottedGeometry.GetLocalSize().X, BarPositions->BarMax - BarPositions->BarNearPlane - Style->SliceHandleSize.Y),
+					FSlateLayoutTransform(FVector2f(0.0f, BarPositions->BarNearPlane + Style->SliceHandleSize.Y)
+				)).ToPaintGeometry(),
+				&Style->BackgroundBrush,
+				ESlateDrawEffect::None,
+				Style->BackgroundBrush.GetTint(InWidgetStyle).CopyWithNewOpacity(0.5f)
+			);
+		}
 	}
 
 	return LayerId;
@@ -395,7 +559,24 @@ int32 SDepthBar::OnPaint(const FPaintArgs& Args, const FGeometry& AllottedGeomet
 
 FVector2D SDepthBar::ComputeDesiredSize(float LayoutScaleMultiplier) const
 {
-	return FVector2D(20.0, 800.0);
+	return FVector2D(Style->SliceWidth + Style->Padding.GetTotalSpaceAlong<Orient_Horizontal>(), 800.0);
+}
+
+FMargin SDepthBar::GetBarInset() const
+{
+	FMargin Inset = Style->Padding;
+	
+	// Allow space for the plane buttons & their paddings
+	Inset.Top += Style->PlaneButtonSize.Y;
+	Inset.Bottom += Style->PlaneButtonSize.Y;
+	Inset.Top += Style->FarPlaneButtonPadding.GetTotalSpaceAlong<Orient_Vertical>();
+	Inset.Bottom += Style->NearPlaneButtonPadding.GetTotalSpaceAlong<Orient_Vertical>();
+	
+	// Allow space for the top & bottom handles
+	Inset.Top += Style->SliceHandleSize.X;
+	Inset.Bottom += Style->SliceHandleSize.Y;
+	
+	return Inset;
 }
 
 TOptional<SDepthBar::FBarPositions> SDepthBar::GetBarPositions(const FGeometry& InGeometry) const
@@ -406,7 +587,7 @@ TOptional<SDepthBar::FBarPositions> SDepthBar::GetBarPositions(const FGeometry&
 		return TOptional<FBarPositions>();
 	}
 	
-	const FMargin Inset = Padding.Get();
+	const FMargin Inset = GetBarInset();
 	
 	if (Mode.Get() == EMode::Orthographic)
 	{
@@ -511,21 +692,23 @@ double SDepthBar::FBarPositions::BarToAlignedPosition(float BarPosition) const
 
 EVisibility SDepthBar::GetFarPlaneLabelVisibility() const
 {
-	return HoverTarget == EDragTarget::Far || HoverTarget == EDragTarget::Slice ? EVisibility::Visible : EVisibility::Collapsed;
+	return HoverTarget == EDragTarget::Far_Handle || HoverTarget == EDragTarget::Slice ? EVisibility::Visible : EVisibility::Collapsed;
 }
 
 EVisibility SDepthBar::GetNearPlaneLabelVisibility() const
 {
-	return HoverTarget == EDragTarget::Near || HoverTarget == EDragTarget::Slice ? EVisibility::Visible : EVisibility::Collapsed;
+	return HoverTarget == EDragTarget::Near_Handle || HoverTarget == EDragTarget::Slice ? EVisibility::Visible : EVisibility::Collapsed;
 }
 
 FText SDepthBar::GetFarPlaneText() const
 {
 	if (Mode.Get() == EMode::Orthographic)
 	{
-		if (const FDepthSpace* Space = DepthSpace.Get().GetPtrOrNull())
+		TOptional<FDepthSpace> OptionalDepthSpace = DepthSpace.Get();
+		if (const FDepthSpace* Space = OptionalDepthSpace.GetPtrOrNull())
 		{
-			if (const double* Far = FarPlane.Get().GetPtrOrNull())
+			TOptional<double> OptionalFarPlane = FarPlane.Get(); 
+			if (const double* Far = OptionalFarPlane.GetPtrOrNull())
 			{
 				const double RelativePosition = Space->AlignedToRelativePosition(*Far);
 				const double WorldPosition = Space->RelativeToBoundsPosition(RelativePosition);
@@ -542,9 +725,11 @@ FText SDepthBar::GetNearPlaneText() const
 {
 	if (Mode.Get() == EMode::Orthographic)
 	{
-		if (const FDepthSpace* Space = DepthSpace.Get().GetPtrOrNull())
+		TOptional<FDepthSpace> OptionalDepthSpace = DepthSpace.Get();
+		if (const FDepthSpace* Space = OptionalDepthSpace.GetPtrOrNull())
 		{
-			if (const double* Near = NearPlane.Get().GetPtrOrNull())
+			TOptional<double> OptionalNearPlane = NearPlane.Get(); 
+			if (const double* Near = OptionalNearPlane.GetPtrOrNull())
 			{
 				const double RelativePosition = Space->AlignedToRelativePosition(*Near);
 				const double WorldPosition = Space->RelativeToBoundsPosition(RelativePosition);
diff --git a/Engine/Source/Editor/EditorFramework/Public/Styling/DepthBarStyle.h b/Engine/Source/Editor/EditorFramework/Public/Styling/DepthBarStyle.h
new file mode 100644
--- /dev/null
+++ b/Engine/Source/Editor/EditorFramework/Public/Styling/DepthBarStyle.h
@@ -0,0 +1,104 @@
+﻿// Copyright Epic Games, Inc. All Rights Reserved.
+
+#pragma once
+
+#include "CoreMinimal.h"
+#include "UObject/ObjectMacros.h"
+#include "Styling/SlateTypes.h"
+#include "Styling/StyleColors.h"
+
+#include "DepthBarStyle.generated.h"
+
+USTRUCT()
+struct FDepthBarStyle : public FSlateWidgetStyle
+{
+	GENERATED_BODY()
+	
+	EDITORFRAMEWORK_API FDepthBarStyle();
+	EDITORFRAMEWORK_API FDepthBarStyle(const FDepthBarStyle&);
+	EDITORFRAMEWORK_API virtual ~FDepthBarStyle() override;
+	
+	EDITORFRAMEWORK_API virtual void GetResources(TArray<const FSlateBrush*>& OutBrushes) const override;
+
+	EDITORFRAMEWORK_API static const FName TypeName;
+	virtual const FName GetTypeName() const override { return TypeName; };
+	
+	EDITORFRAMEWORK_API static const FDepthBarStyle& GetDefault();
+	
+	UPROPERTY(EditAnywhere, Category = Appearance)
+	FMargin Padding = FMargin(4.0f);
+	FDepthBarStyle& SetPadding(const FMargin& InPadding) { Padding = InPadding; return *this; };
+	
+	UPROPERTY(EditAnywhere, Category = Appearance)
+	FSlateBrush BackgroundBrush;
+	FDepthBarStyle& SetBackgroundBrush(const FSlateBrush& InBackground) { BackgroundBrush = InBackground; return *this; };
+	
+	UPROPERTY(EditAnywhere, Category = Appearance)
+	FSlateBrush TrackBrush;
+	FDepthBarStyle& SetTrack(const FSlateBrush& InTrack) { TrackBrush = InTrack; return *this; };
+	
+	UPROPERTY(EditAnywhere, Category = Appearance)
+	float TrackWidth = 2.0f;
+	FDepthBarStyle& SetTrackWidth(float InTrackWidth) { TrackWidth = InTrackWidth; return *this; };
+
+	UPROPERTY(EditAnywhere, Category = Appearance)
+	FSlateBrush SliceNormalBrush;
+	FDepthBarStyle& SetSliceNormal(const FSlateBrush& InSliceNormal) { SliceNormalBrush = InSliceNormal; return *this; };
+	
+	UPROPERTY(EditAnywhere, Category = Appearance)
+	FSlateBrush SliceHoveredBrush;
+	FDepthBarStyle& SetSliceHovered(const FSlateBrush& InSliceHovered) { SliceHoveredBrush = InSliceHovered; return *this; };
+	
+	UPROPERTY(EditAnywhere, Category = Appearance)
+	float SliceWidth = 16.0f;
+	FDepthBarStyle& SetSliceWidth(float InSliceWidth) { SliceWidth = InSliceWidth; return *this; };
+	
+	UPROPERTY(EditAnywhere, Category = Appearance)
+	FSlateBrush SliceTopBrush;
+	FDepthBarStyle& SetSliceTopBrush(const FSlateBrush& InSliceTopBrush) { SliceTopBrush = InSliceTopBrush; return *this; };
+	
+	UPROPERTY(EditAnywhere, Category = Appearance)
+	FSlateBrush SliceTopHoveredBrush;
+	FDepthBarStyle& SetSliceTopHovered(const FSlateBrush& InSliceTopHovered) { SliceTopHoveredBrush = InSliceTopHovered; return *this; };
+	
+	UPROPERTY(EditAnywhere, Category = Appearance)
+	FSlateBrush SliceBottomBrush;
+	FDepthBarStyle& SetSliceBottom(const FSlateBrush& InSliceBottom) { SliceBottomBrush = InSliceBottom; return *this; };
+	
+	UPROPERTY(EditAnywhere, Category = Appearance)
+	FSlateBrush SliceBottomHoveredBrush;
+	FDepthBarStyle& SetSliceBottomHovered(const FSlateBrush& InSliceBottomHovered) { SliceBottomHoveredBrush = InSliceBottomHovered; return *this; };
+	
+	UPROPERTY(EditAnywhere, Category = Appearance)
+	FVector2f SliceHandleSize = FVector2f(16.0f, 16.0f);
+	FDepthBarStyle& SetSliceHandleSize(const FVector2f& InHandleSize) { SliceHandleSize = InHandleSize; return *this; };
+	
+	UPROPERTY(EditAnywhere, Category = Appearance)
+	FSlateBrush FarPlaneButtonBrush;
+	FDepthBarStyle& SetFarPlaneButtonBrush(const FSlateBrush& InFarPlaneButton) { FarPlaneButtonBrush = InFarPlaneButton; return *this; };
+	
+	UPROPERTY(EditAnywhere, Category = Appearance)
+	FSlateBrush NearPlaneButtonBrush;
+	FDepthBarStyle& SetNearPlaneButtonBrush(const FSlateBrush& InNearPlaneButton) { NearPlaneButtonBrush = InNearPlaneButton; return *this; };
+	
+	UPROPERTY(EditAnywhere, Category = Appearance)
+	FVector2f PlaneButtonSize = FVector2f(16.0f, 16.0f);
+	FDepthBarStyle& SetPlaneButtonSize(const FVector2f& InHandleSize) { PlaneButtonSize = InHandleSize; return *this; };
+	
+	UPROPERTY(EditAnywhere, Category = Appearance)
+	FMargin FarPlaneButtonPadding = FMargin(4.0f);
+	FDepthBarStyle& SetFarPlaneButtonPadding(const FMargin& InPlaneButtonPadding) { FarPlaneButtonPadding = InPlaneButtonPadding; return *this; };
+	
+	UPROPERTY(EditAnywhere, Category = Appearance)
+	FMargin NearPlaneButtonPadding = FMargin(4.0f);
+	FDepthBarStyle& SetNearPlaneButtonPadding(const FMargin& InPlaneButtonPadding) { NearPlaneButtonPadding = InPlaneButtonPadding; return *this; };
+	
+	UPROPERTY(EditAnywhere, Category = Appearance)
+	FSlateColor PlaneButtonNormalColor = EStyleColor::ForegroundHeader;
+	FDepthBarStyle& SetPlaneButtonNormalColor(const FSlateColor& InColor) { PlaneButtonNormalColor = InColor; return *this; };
+	
+	UPROPERTY(EditAnywhere, Category = Appearance)
+	FSlateColor PlaneButtonHoveredColor = EStyleColor::ForegroundHover;
+	FDepthBarStyle& SetPlaneButtonHoveredColor(const FSlateColor& InColor) { PlaneButtonHoveredColor = InColor; return *this; };
+	
+};
diff --git a/Engine/Source/Editor/EditorFramework/Public/Viewports/SDepthBar.h b/Engine/Source/Editor/EditorFramework/Public/Viewports/SDepthBar.h
--- a/Engine/Source/Editor/EditorFramework/Public/Viewports/SDepthBar.h
+++ b/Engine/Source/Editor/EditorFramework/Public/Viewports/SDepthBar.h
@@ -6,6 +6,9 @@
 #include "SlotBase.h"
 #include "Widgets/DeclarativeSyntaxSupport.h"
 #include "Layout/Children.h"
+#include "Styling/DepthBarStyle.h"
+#include "Styling/ISlateStyle.h"
+#include "Styling/SlateWidgetStyleAsset.h"
 
 /**
  * A bar that allows for visualization and manipulation of a viewport's depth
@@ -59,11 +62,11 @@ public:
 	/**
 	 * A slot that is placed along the depth bar, representing a range within the depth space
 	 */
-	class FDepthIndicatorSlot : public TSlotBase<FDepthIndicatorSlot>
+	class FDepthIndicatorSlot : public TBasicLayoutWidgetSlot<FDepthIndicatorSlot>
 	{
 	public:
 	
-		SLATE_SLOT_BEGIN_ARGS(FDepthIndicatorSlot, TSlotBase<FDepthIndicatorSlot>)
+		SLATE_SLOT_BEGIN_ARGS(FDepthIndicatorSlot, TBasicLayoutWidgetSlot<FDepthIndicatorSlot>)
 			/** Determines where in the bar the slot's contents appear. */
 			SLATE_ATTRIBUTE(FBoxSphereBounds, Bounds)
 		SLATE_SLOT_END_ARGS()
@@ -75,7 +78,7 @@ public:
 		
 		void Construct(const FChildren& SlotOwner, FSlotArguments&& InArgs)
 		{
-			TSlotBase<FDepthIndicatorSlot>::Construct(SlotOwner, MoveTemp(InArgs));
+			TBasicLayoutWidgetSlot<FDepthIndicatorSlot>::Construct(SlotOwner, MoveTemp(InArgs));
 			Bounds = InArgs._Bounds;
 		}
 		
@@ -104,6 +107,8 @@ public:
 		SLATE_EVENT(FOnDepthChanged, OnFarPlaneChanged)
 		
 		SLATE_SLOT_ARGUMENT(SDepthBar::FDepthIndicatorSlot, Slots)
+		
+		SLATE_STYLE_ARGUMENT(FDepthBarStyle, Style)
 	SLATE_END_ARGS()
 	
 	using FScopedWidgetDepthIndicatorSlotArguments = TPanelChildren<FDepthIndicatorSlot>::FScopedWidgetSlotArguments;
@@ -125,10 +130,6 @@ public:
 		: DepthIndicatorChildren(this)
 		, DepthLabelChildren(this)
 		, AllChildren(this)
-		, Padding(FMargin(0.0f, 4.0f))
-		, SliderMargin(4.0f)
-		, SliderHeadMargin(2.0f)
-		, SliderHeadHeight(3.0f)
 	{
 		AllChildren.AddChildren(DepthIndicatorChildren);
 		AllChildren.AddChildren(DepthLabelChildren);
@@ -142,6 +143,7 @@ public:
 	EDITORFRAMEWORK_API virtual FReply OnMouseButtonDown(const FGeometry& MyGeometry, const FPointerEvent& MouseEvent) override;
 	EDITORFRAMEWORK_API virtual FReply OnMouseMove(const FGeometry& MyGeometry, const FPointerEvent& MouseEvent) override;
 	EDITORFRAMEWORK_API virtual FReply OnMouseButtonUp(const FGeometry& MyGeometry, const FPointerEvent& MouseEvent) override;
+	EDITORFRAMEWORK_API virtual FReply OnMouseButtonDoubleClick(const FGeometry& InMyGeometry, const FPointerEvent& InMouseEvent) override;
 	EDITORFRAMEWORK_API virtual void OnMouseLeave(const FPointerEvent& MouseEvent) override;
 	
 protected:
@@ -155,7 +157,9 @@ private:
 	{
 		FDepthSpace DepthSpace;
 	
+		/** The Top of the bar (i.e. the far plane max) */
 		float BarMin;
+		/** The Bottom of the bar (i.e. the near plane min) */
 		float BarMax;
 		
 		float BarFarPlane;
@@ -217,6 +221,9 @@ private:
 		return FScopedWidgetDepthLabelSlotArguments(MakeUnique<FDepthLabelSlot>(), this->DepthLabelChildren, Index);
 	}
 
+	/** Gets the inset space where bar drawing occurs */
+	FMargin GetBarInset() const;
+	
 	TOptional<FBarPositions> GetBarPositions(const FGeometry& InGeometry) const;
 	
 	EVisibility GetFarPlaneLabelVisibility() const;
@@ -231,15 +238,18 @@ private:
 	enum class EDragTarget
 	{
 		None,
-		Near,
-		Far,
+		
+		Near_Button,
+		Far_Button,
+		
+		Near_Handle,
+		Far_Handle,
 		Slice
 	};
 	EDragTarget HoverTarget = EDragTarget::None;
 	EDragTarget DragTarget = EDragTarget::None;
 	float DragOffset = 0.0f;
 	
-	TAttribute<FMargin> Padding;
 	TAttribute<EMode> Mode;
 	TAttribute<TOptional<FDepthSpace>> DepthSpace;
 	TAttribute<TOptional<double>> NearPlane;
@@ -247,14 +257,6 @@ private:
 	FOnDepthChanged OnNearPlaneChanged;
 	FOnDepthChanged OnFarPlaneChanged;
 	
-	float SliderMargin;
-	float SliderHeadMargin;
-	float SliderHeadHeight;
-	
-	const FSlateBrush* BackgroundBrush = nullptr;
-	const FSlateBrush* SliceBrush = nullptr;
-	const FSlateBrush* HoveredSliceBrush = nullptr;
-	const FSlateBrush* HoveredSliceTopBrush = nullptr;
-	const FSlateBrush* HoveredSliceBottomBrush = nullptr;
+	const FDepthBarStyle* Style = nullptr;
 };
 
diff --git a/Engine/Source/Editor/EditorStyle/EditorStyle.Build.cs b/Engine/Source/Editor/EditorStyle/EditorStyle.Build.cs
--- a/Engine/Source/Editor/EditorStyle/EditorStyle.Build.cs
+++ b/Engine/Source/Editor/EditorStyle/EditorStyle.Build.cs
@@ -26,6 +26,7 @@ public class EditorStyle : ModuleRules
 			new string[] {
 				"SlateCore",
 				"UnrealEd",
+				"EditorFramework"
 			}
 		);
 
diff --git a/Engine/Source/Editor/EditorStyle/Private/StarshipStyle.cpp b/Engine/Source/Editor/EditorStyle/Private/StarshipStyle.cpp
--- a/Engine/Source/Editor/EditorStyle/Private/StarshipStyle.cpp
+++ b/Engine/Source/Editor/EditorStyle/Private/StarshipStyle.cpp
@@ -12,6 +12,7 @@
 #if (WITH_EDITOR || (IS_PROGRAM && PLATFORM_DESKTOP))
 	#include "PlatformInfo.h"
 #endif
+#include "Styling/DepthBarStyle.h"
 #include "Styling/ToolBarStyle.h"
 #include "Styling/SegmentedControlStyle.h"
 #include "Styling/StyleColors.h"
@@ -3246,12 +3247,39 @@ void FStarshipEditorStyle::FStyle::SetupViewportStyles()
 	
 	// Viewport Depth Bar
 	{
-		Set("EditorViewport.DepthBar.Background", new FSlateRoundedBoxBrush(FStyleColors::Panel, 4.0f));
-		Set("EditorViewport.DepthBar.Slice", new FSlateColorBrush(FStyleColors::Hover));
-		Set("EditorViewport.DepthBar.SliceHovered", new FSlateColorBrush(FStyleColors::Hover2));
-		Set("EditorViewport.DepthBar.SliceTop", new FSlateRoundedBoxBrush(FStyleColors::Hover2, FVector4(2.0f, 2.0f, 0.0f, 0.0f)));
-		Set("EditorViewport.DepthBar.SliceBottom", new FSlateRoundedBoxBrush(FStyleColors::Hover2, FVector4(0.0f, 0.0f, 2.0f, 2.0f)));
+		FDepthBarStyle DepthBarStyle = FDepthBarStyle()
+			.SetBackgroundBrush(FSlateRoundedBoxBrush(FStyleColors::Panel, 4.0f))
+			.SetTrack(FSlateRoundedBoxBrush(FStyleColors::Input, 1.0f))
+			.SetSliceNormal(FSlateColorBrush(FStyleColors::Dropdown))
+			.SetSliceHovered(FSlateColorBrush(FStyleColors::Hover))
+			.SetSliceTopBrush(FSlateRoundedBoxBrush(FStyleColors::DropdownOutline, FVector4(2.0f, 2.0f, 0.0f, 0.0f)))
+			.SetSliceTopHovered(FSlateRoundedBoxBrush(FStyleColors::Hover2, FVector4(2.0f, 2.0f, 0.0f, 0.0f)))
+			.SetSliceBottom(FSlateRoundedBoxBrush(FStyleColors::DropdownOutline, FVector4(0.0f, 0.0f, 2.0f, 2.0f)))
+			.SetSliceBottomHovered(FSlateRoundedBoxBrush(FStyleColors::Hover2, FVector4(0.0f, 0.0f, 2.0f, 2.0f)))
+			.SetFarPlaneButtonBrush(IMAGE_BRUSH_SVG("Starship/MainToolbar/landscape", Icon16x16))
+			.SetNearPlaneButtonBrush(IMAGE_BRUSH_SVG("Starship/AssetIcons/Camera_16", Icon16x16))
+			.SetFarPlaneButtonPadding(FMargin(4.0f, 4.0f, 4.0f, 8.0f))
+			.SetNearPlaneButtonPadding(FMargin(4.0f, 8.0f, 4.0f, 0.0f))
+		;
+		Set("EditorViewport.DepthBar", DepthBarStyle);
 		
+		FToolBarStyle DepthBarToolbarStyle = GetWidgetStyle<FToolBarStyle>("ViewportToolbar.Raised");
+		DepthBarToolbarStyle
+			.SetBackgroundPadding(FMargin(0.0f))
+			.SetButtonPadding(FMargin(2.0f, 4.0f))
+			.SetComboButtonPadding(FMargin(2.0f, 4.0f))
+			.SetSeparatorPadding(FMargin(2.0f, 4.0f))
+			.SetShowLabels(false)
+		;
+		
+		DepthBarToolbarStyle.ComboButtonStyle
+			.SetDownArrowImage(FSlateNoResource())
+			.SetDownArrowPadding(FMargin(0.0f))
+		;
+		
+		Set("EditorViewport.DepthBar.Toolbar", DepthBarToolbarStyle);
+		
+		Set("EditorViewport.DepthBar.Background", new FSlateRoundedBoxBrush(FStyleColors::Panel, 8.0f, COLOR("555555FF"), 1.0f));
 		Set("EditorViewport.DepthBar.SelectionHighlight", new FSlateColorBrush(FSlateColor(COLOR("EBA20AFF"))));
 	}
 }
@@ -4927,6 +4955,10 @@ void FStarshipEditorStyle::FStyle::SetupLevelEditorStyle()
 		Set( "EditorViewport.Bottom",      new IMAGE_BRUSH_SVG("Starship/Common/ViewBottom", Icon16x16 ) );
 		Set( "EditorViewport.Right",       new IMAGE_BRUSH_SVG("Starship/Common/ViewRight", Icon16x16 ) );
 		Set( "EditorViewport.Back",        new IMAGE_BRUSH_SVG("Starship/Common/ViewBack", Icon16x16 ) );
+		
+		Set( "EditorViewport.UseOrthographicClippingPlanes", new IMAGE_BRUSH_SVG("Starship/Common/MediaPlanes", Icon16x16) );
+		Set( "EditorViewport.OrthographicClippingSettings", new IMAGE_BRUSH_SVG("Starship/Common/tool-properties", Icon16x16) );
+		Set( "EditorViewport.DismissOrthographicClippingPlanes", new IMAGE_BRUSH_SVG("Starship/Common/cross-16", Icon16x16) );
 
 		Set("EditorViewport.ToggleStats", new IMAGE_BRUSH_SVG("Starship/Common/Statistics", Icon16x16));
 		Set("EditorViewport.ToggleFPS", new IMAGE_BRUSH_SVG("Starship/Common/FPS", Icon16x16));
diff --git a/Engine/Source/Editor/LevelEditor/Private/SLevelViewport.cpp b/Engine/Source/Editor/LevelEditor/Private/SLevelViewport.cpp
--- a/Engine/Source/Editor/LevelEditor/Private/SLevelViewport.cpp
+++ b/Engine/Source/Editor/LevelEditor/Private/SLevelViewport.cpp
@@ -3642,11 +3642,12 @@ void SLevelViewport::OnElementSelectionChanged(const UTypedElementSelectionSet*
 				
 			DepthBar->AddDepthIndicatorSlot()
 			.Bounds(Bounds)
+			.HAlign(HAlign_Center)
 			[
 				SNew(SImage)
 				.Image(FAppStyle::Get().GetBrush("EditorViewport.DepthBar.SelectionHighlight"))
 				.ToolTipText(DisplayName)
-				.DesiredSizeOverride(FVector2d(4, 4))
+				.DesiredSizeOverride(FVector2d(8, 4))
 			];
 		}
 	}
diff --git a/Engine/Source/Editor/UnrealEd/Private/EditorViewportCommands.cpp b/Engine/Source/Editor/UnrealEd/Private/EditorViewportCommands.cpp
--- a/Engine/Source/Editor/UnrealEd/Private/EditorViewportCommands.cpp
+++ b/Engine/Source/Editor/UnrealEd/Private/EditorViewportCommands.cpp
@@ -197,6 +197,7 @@ void FEditorViewportCommands::RegisterCommands()
 	UI_COMMAND( FocusOutlinerToContextFolder, "Frame Current Folder in Outliner", "Focuses the current folder (if it's set) in the outliner", EUserInterfaceActionType::Button, FInputChord() );
 
 	UI_COMMAND( UseOrthographicClippingPlanes, "Ortho Clipping Planes", "Whether or not the depth of the orthographic view should be clipped.", EUserInterfaceActionType::Check, FInputChord() );
+	UI_COMMAND( DismissOrthographicClippingPlanes, "Dismiss Ortho Clipping Planes", "Resets the orthographic clipping planes and hides the visualization of viewport depth.", EUserInterfaceActionType::Button, FInputChord() );
 
 	UI_COMMAND( LocationGridSnap, "Grid Snap", "Toggle whether objects snap to a grid when you move them.\nUse the submenu to choose the grid size.", EUserInterfaceActionType::ToggleButton, FInputChord() );
 	UI_COMMAND( RotationGridSnap, "Rotation Snap", "Toggle whether objects snap to preset angles when you rotate them.\nUse the submenu to choose the angle increment.", EUserInterfaceActionType::ToggleButton, FInputChord() );
diff --git a/Engine/Source/Editor/UnrealEd/Private/SEditorViewport.cpp b/Engine/Source/Editor/UnrealEd/Private/SEditorViewport.cpp
--- a/Engine/Source/Editor/UnrealEd/Private/SEditorViewport.cpp
+++ b/Engine/Source/Editor/UnrealEd/Private/SEditorViewport.cpp
@@ -38,6 +38,7 @@
 #include "LevelEditorViewport.h"
 #include "Viewports/SDepthBar.h"
 #include "SEditorViewportGridPanel.h"
+#include "ToolMenus.h"
 #include "ViewportInteractions/ViewportInteractionsBehaviorSource.h"
 #include "Widgets/Input/SButton.h"
 
@@ -45,7 +46,7 @@
 
 static TAutoConsoleVariable<float> CVarEditorViewportDepthBarMinSpace(
 	TEXT("r.Editor.DepthBarMinSpace"),
-	600.0f,
+	500.0f,
 	TEXT("The minimum vertical space to show the depth bar. Defaults to 600."),
 	ECVF_Default);
 
@@ -292,47 +293,20 @@ void SEditorViewport::Construct( const FArguments& InArgs )
 	}
 	// clang-format on
 	
-	ViewportOverlay->AddSlot()
-	.HAlign(HAlign_Right)	
-	.VAlign(VAlign_Fill)
-	[
-		SNew(SBox)
-		.Padding(0.0f, 100.0f, 6.0f, 100.0f)
-		.Visibility(this, &SEditorViewport::OnGetDepthBarVisibility)
+	if (TSharedPtr<SWidget> DepthBarWidget = BuildViewportDepthBar())
+	{
+		ViewportOverlay->AddSlot()
+		.HAlign(HAlign_Right)	
+		.VAlign(VAlign_Fill)
 		[
-			SNew(SVerticalBox)
-			+ SVerticalBox::Slot()
-			.FillHeight(1.0f)
+			SNew(SBox)
+			.Padding(0.0f, 100.0f, 6.0f, 100.0f)
+			.Visibility(this, &SEditorViewport::OnGetDepthBarVisibility)
 			[
-				SAssignNew(DepthBar, SDepthBar)
-				.Mode(this, &SEditorViewport::GetDepthBarMode)
-				.DepthSpace(this, &SEditorViewport::GetDepthBarSpace)
-				.NearPlane(this, &SEditorViewport::GetDepthBarNearPlane)
-				.FarPlane(this, &SEditorViewport::GetDepthBarFarPlane)
-				.OnNearPlaneChanged(this, &SEditorViewport::OnDepthBarNearPlaneChanged)
-				.OnFarPlaneChanged(this, &SEditorViewport::OnDepthBarFarPlaneChanged)
+				DepthBarWidget.ToSharedRef()
 			]
-			+ SVerticalBox::Slot()
-			.Padding(0.0f, 6.0f, 0.0f, 0.0f)
-			.AutoHeight()
-			[
-				SNew(SButton)
-				.OnClicked(this, &SEditorViewport::DismissDepthBar)
-				.ToolTipText(LOCTEXT("ResetViewportDepth", "Reset Viewport Depth"))
-				.ContentPadding(FMargin(4.0f))
-				.NormalPaddingOverride(FMargin(0.0f))
-				.PressedPaddingOverride(FMargin(0.0f))
-				.HAlign(HAlign_Center)
-				.VAlign(VAlign_Center)
-				[
-					SNew(SImage)
-					.Image(FAppStyle::GetBrush("Cross"))
-					.DesiredSizeOverride(FVector2D(12.0f, 12.0f))
-				]
-			]
-		]
-	];
-
+		];
+	}
 	
 	PopulateViewportOverlays(ViewportOverlay.ToSharedRef());
 
@@ -606,6 +580,11 @@ void SEditorViewport::BindCommands()
 		FCanExecuteAction::CreateSP( this, &SEditorViewport::CanUseDepthBar ),
 		FIsActionChecked::CreateSP( this, &SEditorViewport::IsDepthBarActive )
 		);
+		
+	CommandListRef.MapAction(
+		Commands.DismissOrthographicClippingPlanes,
+		FExecuteAction::CreateSP( this, &SEditorViewport::DismissDepthBar )
+		);
 	
 	CommandListRef.MapAction(
 		Commands.SurfaceSnapping,
@@ -816,12 +795,11 @@ void SEditorViewport::OnDepthBarFarPlaneChanged(const TOptional<double>& InFarPl
 	Client->Invalidate();
 }
 
-FReply SEditorViewport::DismissDepthBar()
+void SEditorViewport::DismissDepthBar()
 {
 	Client->SetOrthographicFarPlaneOverride(TOptional<double>());
 	Client->SetOrthographicNearPlaneOverride(TOptional<double>());
 	bShowDepthBar = false;
-	return FReply::Handled();
 }
 
 void SEditorViewport::OnToggleRealtime()
@@ -1170,6 +1148,55 @@ void SEditorViewport::OnCycleCoordinateSystem()
 	Client->SetWidgetCoordSystemSpace( (ECoordSystem)CoordSystemAsInt );
 }
 
+TSharedPtr<SWidget> SEditorViewport::BuildViewportDepthBar()
+{
+	const FName MenuName = "EditorViewportDepthBarMenu";
+	
+	if (!UToolMenus::Get()->IsMenuRegistered(MenuName))
+	{
+		UToolMenu* Menu = UToolMenus::Get()->RegisterMenu(MenuName, NAME_None, EMultiBoxType::VerticalToolBar);
+		Menu->StyleName = "EditorViewport.DepthBar.Toolbar";
+	
+		FToolMenuSection& Actions = Menu->AddSection("Actions");
+		Actions.AddSeparator("Start");
+		Actions.AddEntry(UE::UnrealEd::CreateOrthographicClippingPlanesSubmenu());
+		
+		FToolMenuSection& EndSection = Menu->AddSection("End", {}, FToolMenuInsert(NAME_None, EToolMenuInsertType::Last));
+		EndSection.AddMenuEntry(FEditorViewportCommands::Get().DismissOrthographicClippingPlanes);
+	}
+	
+	FToolMenuContext Context;
+	
+	UUnrealEdViewportToolbarContext* ContextObject = UE::UnrealEd::CreateViewportToolbarDefaultContext(SharedThis(this));
+	Context.AddObject(ContextObject);
+	
+	Context.AppendCommandList(GetCommandList());
+	
+	return SNew(SBorder)
+		.BorderImage(FAppStyle::GetBrush("EditorViewport.DepthBar.Background"))
+		.Padding(3.0f)
+		[
+			SNew(SVerticalBox)
+			+ SVerticalBox::Slot()
+			.FillHeight(1.0f)
+			[
+				SAssignNew(DepthBar, SDepthBar)
+				.Style(FAppStyle::Get(), "EditorViewport.DepthBar")
+				.Mode(this, &SEditorViewport::GetDepthBarMode)
+				.DepthSpace(this, &SEditorViewport::GetDepthBarSpace)
+				.NearPlane(this, &SEditorViewport::GetDepthBarNearPlane)
+				.FarPlane(this, &SEditorViewport::GetDepthBarFarPlane)
+				.OnNearPlaneChanged(this, &SEditorViewport::OnDepthBarNearPlaneChanged)
+				.OnFarPlaneChanged(this, &SEditorViewport::OnDepthBarFarPlaneChanged)
+			]
+			+ SVerticalBox::Slot()
+			.AutoHeight()
+			[
+				UToolMenus::Get()->GenerateWidget(MenuName, Context)
+			]
+		];
+}
+
 void SEditorViewport::OnToggleUseDepthBar()
 {
 	if (Client->IsOrtho())
diff --git a/Engine/Source/Editor/UnrealEd/Private/ViewportToolbar/UnrealEdViewportToolbar.cpp b/Engine/Source/Editor/UnrealEd/Private/ViewportToolbar/UnrealEdViewportToolbar.cpp
--- a/Engine/Source/Editor/UnrealEd/Private/ViewportToolbar/UnrealEdViewportToolbar.cpp
+++ b/Engine/Source/Editor/UnrealEd/Private/ViewportToolbar/UnrealEdViewportToolbar.cpp
@@ -793,7 +793,7 @@ public:
 		[
 			SNew(SBox)
 			.Padding(FMargin(4.0f, 0.0f, 0.0f, 0.0f))
-			.WidthOverride(100.0f)
+			.HAlign(HAlign_Right)
 			[
 				SNew(SBorder)
 				.BorderImage(FAppStyle::Get().GetBrush("Menu.WidgetBorder"))
@@ -808,6 +808,7 @@ public:
 					.Value(this, &SOrthographicViewPlaneSlider::GetValue)
 					.OnValueChanged(this, &SOrthographicViewPlaneSlider::OnValueChanged)
 					.MaxFractionalDigits(2)
+					.MinDesiredWidth(100)
 				]
 			]
 		];
@@ -4320,10 +4321,15 @@ FToolMenuEntry CreateOrthographicClippingPlanesSubmenu()
 				Viewport->GetViewportClient(),
 				UE::UnrealEd::Private::SOrthographicViewPlaneSlider::EMode::Near
 			);
+				
+			FToolMenuEntry& FarPlaneEntry = ClippingSection.AddEntry(FToolMenuEntry::InitWidget("FarPlane", FarSlider, LOCTEXT("OrthoFarPlaneLabel", "Far Clip Location")));
+			FarPlaneEntry.Icon = FSlateIcon(FAppStyle::GetAppStyleSetName(), "ViewportToolbar.FarViewPlane");
 			
-			ClippingSection.AddEntry(FToolMenuEntry::InitWidget("FarPlane", FarSlider, LOCTEXT("OrthoFarPlaneLabel", "Far Plane")));
-			ClippingSection.AddEntry(FToolMenuEntry::InitWidget("NearPlane", NearSlider, LOCTEXT("OrthoNearPlaneLabel", "Near Plane")));
-		})
+			FToolMenuEntry& NearPlaneEntry = ClippingSection.AddEntry(FToolMenuEntry::InitWidget("NearPlane", NearSlider, LOCTEXT("OrthoNearPlaneLabel", "Near Clip Location")));
+			NearPlaneEntry.Icon = FSlateIcon(FAppStyle::GetAppStyleSetName(), "ViewportToolbar.NearViewPlane");
+		}),
+		false,
+		FSlateIcon(FAppStyle::GetAppStyleSetName(), "EditorViewport.OrthographicClippingSettings")	
 	);
 }
 
diff --git a/Engine/Source/Editor/UnrealEd/Public/EditorViewportCommands.h b/Engine/Source/Editor/UnrealEd/Public/EditorViewportCommands.h
--- a/Engine/Source/Editor/UnrealEd/Public/EditorViewportCommands.h
+++ b/Engine/Source/Editor/UnrealEd/Public/EditorViewportCommands.h
@@ -230,6 +230,7 @@ public:
 	TSharedPtr< FUICommandInfo > FocusOutlinerToContextFolder;
 	
 	TSharedPtr<FUICommandInfo> UseOrthographicClippingPlanes;
+	TSharedPtr<FUICommandInfo> DismissOrthographicClippingPlanes;
 
 	/** Toggle automatic exposure */
 	TSharedPtr< FUICommandInfo > ToggleAutoExposure;
diff --git a/Engine/Source/Editor/UnrealEd/Public/SEditorViewport.h b/Engine/Source/Editor/UnrealEd/Public/SEditorViewport.h
--- a/Engine/Source/Editor/UnrealEd/Public/SEditorViewport.h
+++ b/Engine/Source/Editor/UnrealEd/Public/SEditorViewport.h
@@ -271,6 +271,11 @@ protected:
 	 * Called when the user wants to focus the viewport to the current selection
 	 */
 	virtual void OnFocusViewportToSelection(){}
+
+	/**
+	 * Creates the viewport depth bar and any containing widgets
+	 */
+	TSharedPtr<SWidget> BuildViewportDepthBar();
 	
 	/**
 	 * Called when enabling or disabling the depth bar.
@@ -403,7 +408,7 @@ private:
 	void OnDepthBarFarPlaneChanged(const TOptional<double>& InFarPlane);
 	
 	/** Reset the depth bar back to the default state. */
-	FReply DismissDepthBar();
+	void DismissDepthBar();
 
 private:
 	/** The handle to the active timer */
