[Depth Bar] Add support for Ctrl+Drag for higher precision movement of drag handles.
- Enable the Depth Bar by default.
- Mark the Depth Bar as experimental & move some protected functions to private.
- Fix visual issues present in lit mode when depth sliders are set to infinite.

#rnx

diff --git a/Engine/Source/Editor/EditorFramework/Private/Viewports/SDepthBar.cpp b/Engine/Source/Editor/EditorFramework/Private/Viewports/SDepthBar.cpp
--- a/Engine/Source/Editor/EditorFramework/Private/Viewports/SDepthBar.cpp
+++ b/Engine/Source/Editor/EditorFramework/Private/Viewports/SDepthBar.cpp
@@ -136,11 +136,23 @@ FReply SDepthBar::OnMouseButtonDown(const FGeometry& MyGeometry, const FPointerE
 			const FVector2f MousePosition = MyGeometry.AbsoluteToLocal(MouseEvent.GetScreenSpacePosition());
 		
 			DragTarget = HoverTarget;
-			if (DragTarget == EDragTarget::Slice)
+			
+			switch (DragTarget)
 			{
-				DragOffset = MousePosition.Y - Positions->BarFarPlane;
+			case EDragTarget::Far_Handle:
+			case EDragTarget::Slice:
+				DragPosition = FVector2f(0.0f, Positions->BarFarPlane);
+				break;
+			case EDragTarget::Near_Handle:
+				DragPosition = FVector2f(0.0f, Positions->BarNearPlane);
+				break;
+			default:
+				DragPosition = MousePosition;
+				break;
 			}
 			
+			StartDragPosition = DragPosition;
+			
 			return FReply::Handled().CaptureMouse(AsShared());
 		}
 	}
@@ -158,13 +170,26 @@ FReply SDepthBar::OnMouseMove(const FGeometry& MyGeometry, const FPointerEvent&
 	{
 		if (const FBarPositions* Positions = OptionalPositions.GetPtrOrNull())
 		{
-			const FVector2f MousePosition = MyGeometry.AbsoluteToLocal(MouseEvent.GetScreenSpacePosition());
+			float DragRatio = 1.0f;
 			
+			if (MouseEvent.IsAltDown())
+			{
+				// Holding alt down should give consistent world-space movement (1px = 1u), but not if doing so would _lose_ precision.
+				DragRatio = FMath::Min(1.0f, (Positions->BarMax - Positions->BarMin) / static_cast<float>(Positions->DepthSpace.GetAlignedMax() - Positions->DepthSpace.GetAlignedMin()));
+			}
+			
+			if (MouseEvent.IsControlDown())
+			{
+				DragRatio *= 0.1f;
+			}
+			
+			DragPosition += MouseEvent.GetCursorDelta() * DragRatio;
+		
 			if (Mode.Get() == EMode::Orthographic)
 			{
 				if (DragTarget == EDragTarget::Far_Handle)
 				{
-					const float ClampedBarPosition = FMath::Clamp(MousePosition.Y, Positions->BarMin, Positions->BarNearPlane - MinIndicatorDifference);
+					const float ClampedBarPosition = FMath::Clamp(DragPosition.Y, Positions->BarMin, Positions->BarNearPlane - MinIndicatorDifference);
 					
 					if (ClampedBarPosition == Positions->BarMin)
 					{
@@ -179,7 +204,7 @@ FReply SDepthBar::OnMouseMove(const FGeometry& MyGeometry, const FPointerEvent&
 				}
 				if (DragTarget == EDragTarget::Near_Handle)
 				{
-					const float ClampedBarPosition = FMath::Clamp(MousePosition.Y, Positions->BarFarPlane + MinIndicatorDifference, Positions->BarMax);
+					const float ClampedBarPosition = FMath::Clamp(DragPosition.Y, Positions->BarFarPlane + MinIndicatorDifference, Positions->BarMax);
 					if (ClampedBarPosition == Positions->BarMax)
 					{
 						OnNearPlaneChanged.ExecuteIfBound(TOptional<double>());
@@ -193,7 +218,7 @@ FReply SDepthBar::OnMouseMove(const FGeometry& MyGeometry, const FPointerEvent&
 				}
 				if (DragTarget == EDragTarget::Slice)
 				{
-					const float RawPosition = MousePosition.Y - DragOffset;
+					const float RawPosition = DragPosition.Y;
 					const float ClampedBarPosition = FMath::Clamp(RawPosition, Positions->BarMin, Positions->BarMax - (Positions->BarNearPlane - Positions->BarFarPlane));
 					
 					OnFarPlaneChanged.ExecuteIfBound(Positions->BarToAlignedPosition(ClampedBarPosition));
@@ -244,6 +269,19 @@ FReply SDepthBar::OnMouseButtonUp(const FGeometry& MyGeometry, const FPointerEve
 {
 	if (DragTarget != EDragTarget::None)
 	{
+		constexpr float ClickThreshold = 3.0f;
+		if (FVector2f::DistSquared(DragPosition, StartDragPosition) < ClickThreshold)
+		{
+			if (DragTarget == EDragTarget::Far_Button)
+			{
+				OnFarPlaneChanged.ExecuteIfBound(TOptional<double>());
+			}
+			if (DragTarget == EDragTarget::Near_Button)
+			{
+				OnNearPlaneChanged.ExecuteIfBound(TOptional<double>());
+			}
+		}
+	
 		DragTarget = EDragTarget::None;
 		HoverTarget = EDragTarget::None;
 		return FReply::Handled().ReleaseMouseCapture();
diff --git a/Engine/Source/Editor/EditorFramework/Public/Styling/DepthBarStyle.h b/Engine/Source/Editor/EditorFramework/Public/Styling/DepthBarStyle.h
--- a/Engine/Source/Editor/EditorFramework/Public/Styling/DepthBarStyle.h
+++ b/Engine/Source/Editor/EditorFramework/Public/Styling/DepthBarStyle.h
@@ -9,6 +9,7 @@
 
 #include "DepthBarStyle.generated.h"
 
+struct UE_EXPERIMENTAL(5.7, "Represents the style of an SDepthBar") FDepthBarStyle;
 USTRUCT()
 struct FDepthBarStyle : public FSlateWidgetStyle
 {
diff --git a/Engine/Source/Editor/EditorFramework/Public/Viewports/SDepthBar.h b/Engine/Source/Editor/EditorFramework/Public/Viewports/SDepthBar.h
--- a/Engine/Source/Editor/EditorFramework/Public/Viewports/SDepthBar.h
+++ b/Engine/Source/Editor/EditorFramework/Public/Viewports/SDepthBar.h
@@ -13,6 +13,7 @@
 /**
  * A bar that allows for visualization and manipulation of a viewport's depth
  */
+class UE_EXPERIMENTAL(5.7, "SDepthBar represents and controls view depth within viewports.") SDepthBar;
 class SDepthBar : public SPanel
 {
 public:
@@ -248,7 +249,8 @@ private:
 	};
 	EDragTarget HoverTarget = EDragTarget::None;
 	EDragTarget DragTarget = EDragTarget::None;
-	float DragOffset = 0.0f;
+	FVector2f StartDragPosition = FVector2f::Zero();
+	FVector2f DragPosition = FVector2f::Zero();
 	
 	TAttribute<EMode> Mode;
 	TAttribute<TOptional<FDepthSpace>> DepthSpace;
diff --git a/Engine/Source/Editor/UnrealEd/Private/EditorViewportClient.cpp b/Engine/Source/Editor/UnrealEd/Private/EditorViewportClient.cpp
--- a/Engine/Source/Editor/UnrealEd/Private/EditorViewportClient.cpp
+++ b/Engine/Source/Editor/UnrealEd/Private/EditorViewportClient.cpp
@@ -131,7 +131,7 @@ static TAutoConsoleVariable<int32> CVarEditorViewGPUMirrorTest(
 
 static TAutoConsoleVariable<int32> CVarEditorViewportManualOrthographicDepthControl(
 	TEXT("r.Editor.ManualOrthoDepth"),
-	0,
+	1,
 	TEXT("Enable or disable the new orthographic depth tools."),
 	ECVF_RenderThreadSafe);
 
@@ -1366,15 +1366,25 @@ FSceneView* FEditorViewportClient::CalcSceneView(FSceneViewFamily* ViewFamily, c
 				CalculatePlanesViewInfo.OrthoNearClipPlane = -UE_OLD_WORLD_MAX;
 				CalculatePlanesViewInfo.OrthoFarClipPlane = UE_OLD_WORLD_MAX;
 
+				SDepthBar::FDepthSpace DepthSpace(GetViewLocation(), GetForwardVector(), GetWorldBounds());
+				
 				if (OrthographicNearPlane.IsSet())
 				{
 					CalculatePlanesViewInfo.OrthoNearClipPlane = OrthographicNearPlane.GetValue();
 				}
+				else if (!ViewFamily->EngineShowFlags.Wireframe)
+				{
+					CalculatePlanesViewInfo.OrthoNearClipPlane = DepthSpace.GetAlignedMin();
+				}
 				
 				if (OrthographicFarPlane.IsSet())
 				{
 					CalculatePlanesViewInfo.OrthoFarClipPlane = OrthographicFarPlane.GetValue();
 				}
+				else if (!ViewFamily->EngineShowFlags.Wireframe)
+				{
+					CalculatePlanesViewInfo.OrthoFarClipPlane = DepthSpace.GetAlignedMax();
+				}
 				
 				ViewInitOptions.UpdateOrthoPlanes(CalculatePlanesViewInfo);	
 				ZScale = 1.0f / (CalculatePlanesViewInfo.OrthoFarClipPlane - CalculatePlanesViewInfo.OrthoNearClipPlane);
diff --git a/Engine/Source/Editor/UnrealEd/Private/SEditorViewport.cpp b/Engine/Source/Editor/UnrealEd/Private/SEditorViewport.cpp
--- a/Engine/Source/Editor/UnrealEd/Private/SEditorViewport.cpp
+++ b/Engine/Source/Editor/UnrealEd/Private/SEditorViewport.cpp
@@ -47,7 +47,7 @@
 static TAutoConsoleVariable<float> CVarEditorViewportDepthBarMinSpace(
 	TEXT("r.Editor.DepthBarMinSpace"),
 	500.0f,
-	TEXT("The minimum vertical space to show the depth bar. Defaults to 600."),
+	TEXT("The minimum vertical space to show the depth bar. Defaults to 500."),
 	ECVF_Default);
 
 SEditorViewport::SEditorViewport()
@@ -1154,12 +1154,14 @@ TSharedPtr<SWidget> SEditorViewport::BuildViewportDepthBar()
 	
 	if (!UToolMenus::Get()->IsMenuRegistered(MenuName))
 	{
+		using namespace UE::UnrealEd;
+	
 		UToolMenu* Menu = UToolMenus::Get()->RegisterMenu(MenuName, NAME_None, EMultiBoxType::VerticalToolBar);
 		Menu->StyleName = "EditorViewport.DepthBar.Toolbar";
 	
 		FToolMenuSection& Actions = Menu->AddSection("Actions");
 		Actions.AddSeparator("Start");
-		Actions.AddEntry(UE::UnrealEd::CreateOrthographicClippingPlanesSubmenu());
+		Actions.AddEntry(CreateOrthographicClippingPlanesSubmenu(FOrthographicClippingPlanesSubmenuOptions { false } ));
 		
 		FToolMenuSection& EndSection = Menu->AddSection("End", {}, FToolMenuInsert(NAME_None, EToolMenuInsertType::Last));
 		EndSection.AddMenuEntry(FEditorViewportCommands::Get().DismissOrthographicClippingPlanes);
diff --git a/Engine/Source/Editor/UnrealEd/Private/ViewportToolbar/UnrealEdViewportToolbar.cpp b/Engine/Source/Editor/UnrealEd/Private/ViewportToolbar/UnrealEdViewportToolbar.cpp
--- a/Engine/Source/Editor/UnrealEd/Private/ViewportToolbar/UnrealEdViewportToolbar.cpp
+++ b/Engine/Source/Editor/UnrealEd/Private/ViewportToolbar/UnrealEdViewportToolbar.cpp
@@ -3619,9 +3619,6 @@ void PopulateCameraMenu(UToolMenu* InMenu, const FViewportCameraMenuOptions& InO
 			CameraFarViewPlane.Icon = FSlateIcon(FAppStyle::GetAppStyleSetName(), "ViewportToolbar.FarViewPlane");
 			CameraFarViewPlane.Visibility = IsPerspectiveAttribute;
 			ViewSection.AddEntry(CameraFarViewPlane);
-			
-			FToolMenuEntry& UseOrthoClippingPlanes = ViewSection.AddMenuEntry(FEditorViewportCommands::Get().UseOrthographicClippingPlanes);
-			UseOrthoClippingPlanes.Visibility = UE::UnrealEd::GetIsOrthographicAttribute(EditorViewport->GetViewportClient());
 					
 			FToolMenuEntry& OrthoSubmenu = ViewSection.AddEntry(CreateOrthographicClippingPlanesSubmenu());
 			OrthoSubmenu.Visibility = UE::UnrealEd::GetIsOrthographicAttribute(EditorViewport->GetViewportClient());
@@ -4204,8 +4201,30 @@ FText GetCameraSpeedTooltip()
 	return LOCTEXT("CameraSpeedSubMenu_ToolTip", "Set the camera speed.\nShortcut: Hold either mouse button and use the scroll wheel.");
 }
 
-FToolMenuEntry CreateOrthographicClippingPlanesSubmenu()
+FToolMenuEntry CreateOrthographicClippingPlanesSubmenu(const FOrthographicClippingPlanesSubmenuOptions& InOptions)
 {
+	FToolUIAction Action;
+
+	if (InOptions.bDisplayAsToggle)
+	{
+		Action.ExecuteAction.BindLambda([](const FToolMenuContext& Context)
+		{
+			if (const FUIAction* Action = Context.GetActionForCommand(FEditorViewportCommands::Get().UseOrthographicClippingPlanes))
+			{
+				Action->Execute();
+			}
+		});
+		
+		Action.GetActionCheckState.BindLambda([](const FToolMenuContext& Context)
+		{
+			if (const FUIAction* Action = Context.GetActionForCommand(FEditorViewportCommands::Get().UseOrthographicClippingPlanes))
+			{
+				return Action->GetCheckState();
+			}
+			return ECheckBoxState::Unchecked;
+		});
+	}
+
 	return FToolMenuEntry::InitSubMenu(
 		"OrthographicClippingPlanes",
 		LOCTEXT("OrthographicClippingPlanesSubMenuLabel", "Ortho Clipping Planes"),
@@ -4238,6 +4257,8 @@ FToolMenuEntry CreateOrthographicClippingPlanesSubmenu()
 			FToolMenuEntry& NearPlaneEntry = ClippingSection.AddEntry(FToolMenuEntry::InitWidget("NearPlane", NearSlider, LOCTEXT("OrthoNearPlaneLabel", "Near Clip Location")));
 			NearPlaneEntry.Icon = FSlateIcon(FAppStyle::GetAppStyleSetName(), "ViewportToolbar.NearViewPlane");
 		}),
+		Action,
+		InOptions.bDisplayAsToggle ? EUserInterfaceActionType::ToggleButton : EUserInterfaceActionType::None,
 		false,
 		FSlateIcon(FAppStyle::GetAppStyleSetName(), "EditorViewport.OrthographicClippingSettings")	
 	);
diff --git a/Engine/Source/Editor/UnrealEd/Public/SEditorViewport.h b/Engine/Source/Editor/UnrealEd/Public/SEditorViewport.h
--- a/Engine/Source/Editor/UnrealEd/Public/SEditorViewport.h
+++ b/Engine/Source/Editor/UnrealEd/Public/SEditorViewport.h
@@ -272,26 +272,6 @@ protected:
 	 */
 	virtual void OnFocusViewportToSelection(){}
 
-	/**
-	 * Creates the viewport depth bar and any containing widgets
-	 */
-	TSharedPtr<SWidget> BuildViewportDepthBar();
-	
-	/**
-	 * Called when enabling or disabling the depth bar.
-	 */
-	void OnToggleUseDepthBar();
-
-	/**
-	 * @return True if the depth bar can be enabled. 
-	 */
-	bool CanUseDepthBar() const;
-
-	/**
-	 * @return True if the depth bar is active 
-	 */
-	bool IsDepthBarActive() const;
-
 	/** Gets the world this viewport is for */
 	UNREALED_API virtual UWorld* GetWorld() const;
 
@@ -361,6 +341,7 @@ protected:
 	TSharedPtr<FEditorViewportClient> Client;
 	
 	/** The depth bar used by this viewport */
+	UE_EXPERIMENTAL(5.7, "Experimental interface to the Depth bar controls")
 	TSharedPtr<SDepthBar> DepthBar;
 	
 	/** Commandlist used in the viewport (Maps commands to viewport specific actions) */
@@ -386,6 +367,26 @@ private:
 	 */
 	void RegisterViewportInteractions();
 	
+	/**
+	 * Creates the viewport depth bar and any containing widgets
+	 */
+	TSharedPtr<SWidget> BuildViewportDepthBar();
+	
+	/**
+	 * Called when enabling or disabling the depth bar.
+	 */
+	void OnToggleUseDepthBar();
+
+	/**
+	 * @return True if the depth bar can be enabled. 
+	 */
+	bool CanUseDepthBar() const;
+
+	/**
+	 * @return True if the depth bar is active 
+	 */
+	bool IsDepthBarActive() const;
+	
 	/** @return The visibility of the viewport depth bar. */
 	EVisibility OnGetDepthBarVisibility() const;
 	
diff --git a/Engine/Source/Editor/UnrealEd/Public/ViewportToolbar/UnrealEdViewportToolbar.h b/Engine/Source/Editor/UnrealEd/Public/ViewportToolbar/UnrealEdViewportToolbar.h
--- a/Engine/Source/Editor/UnrealEd/Public/ViewportToolbar/UnrealEdViewportToolbar.h
+++ b/Engine/Source/Editor/UnrealEd/Public/ViewportToolbar/UnrealEdViewportToolbar.h
@@ -284,7 +284,12 @@ UNREALED_API FToolMenuEntry CreateResetScalabilitySubmenu();
 
 UNREALED_API FText GetCameraSpeedTooltip();
 
-UNREALED_API FToolMenuEntry CreateOrthographicClippingPlanesSubmenu();
+struct FOrthographicClippingPlanesSubmenuOptions
+{
+	bool bDisplayAsToggle = true;
+};
+
+UNREALED_API FToolMenuEntry CreateOrthographicClippingPlanesSubmenu(const FOrthographicClippingPlanesSubmenuOptions& InOptions = {});
 
 /**
  * Returns Visible if the provided Viewport Client is in perspective view, Collapsed if not
