Remove UE_EXPERIMENTAL from DepthBarStyle.

#rnx

diff --git a/Engine/Source/Editor/EditorFramework/Public/Styling/DepthBarStyle.h b/Engine/Source/Editor/EditorFramework/Public/Styling/DepthBarStyle.h
--- a/Engine/Source/Editor/EditorFramework/Public/Styling/DepthBarStyle.h
+++ b/Engine/Source/Editor/EditorFramework/Public/Styling/DepthBarStyle.h
@@ -9,7 +9,6 @@
 
 #include "DepthBarStyle.generated.h"
 
-struct UE_EXPERIMENTAL(5.7, "Represents the style of an SDepthBar") FDepthBarStyle;
 USTRUCT()
 struct FDepthBarStyle : public FSlateWidgetStyle
 {
