Remove UE_EXPERIMENTAL from SDepthBar.

#rnx

diff --git a/Engine/Source/Editor/EditorFramework/Public/Viewports/SDepthBar.h b/Engine/Source/Editor/EditorFramework/Public/Viewports/SDepthBar.h
--- a/Engine/Source/Editor/EditorFramework/Public/Viewports/SDepthBar.h
+++ b/Engine/Source/Editor/EditorFramework/Public/Viewports/SDepthBar.h
@@ -13,7 +13,6 @@
 /**
  * A bar that allows for visualization and manipulation of a viewport's depth
  */
-class UE_EXPERIMENTAL(5.7, "SDepthBar represents and controls view depth within viewports.") SDepthBar;
 class SDepthBar : public SPanel
 {
 public:
