[Depth Bar] Convert all depth bar calculations and tracking values to double precision. This fixes precision issues that cropped up in very large worlds.

#rnx

diff --git a/Engine/Source/Editor/EditorFramework/Private/Viewports/SDepthBar.cpp b/Engine/Source/Editor/EditorFramework/Private/Viewports/SDepthBar.cpp
--- a/Engine/Source/Editor/EditorFramework/Private/Viewports/SDepthBar.cpp
+++ b/Engine/Source/Editor/EditorFramework/Private/Viewports/SDepthBar.cpp
@@ -10,7 +10,7 @@
 
 namespace UE::EditorWidgets::Private
 {
-	constexpr float MinIndicatorDifference = 2.0f;
+	constexpr double MinIndicatorDifference = 2.0;
 	
 	FText GetDistanceDisplayText(double Value)
 	{
@@ -57,7 +57,7 @@ void SDepthBar::Construct(const FArguments& InArgs)
 	
 	AddDepthLabelSlot()
 	.VAlign(VAlign_Top)
-	.OnGetBarPosition_Lambda([](const FBarPositions& BarPositions) { return BarPositions.BarFarPlane; })
+	.OnGetBarPosition_Lambda([](const FBarPositions& BarPositions) { return static_cast<float>(BarPositions.BarFarPlane); })
 	[
 		SNew(SBorder)
 		.BorderImage(&Style->BackgroundBrush)
@@ -71,7 +71,7 @@ void SDepthBar::Construct(const FArguments& InArgs)
 	
 	AddDepthLabelSlot()
 	.VAlign(VAlign_Bottom)
-	.OnGetBarPosition_Lambda([](const FBarPositions& BarPositions) { return BarPositions.BarNearPlane; })
+	.OnGetBarPosition_Lambda([](const FBarPositions& BarPositions) { return static_cast<float>(BarPositions.BarNearPlane); })
 	[
 		SNew(SBorder)
 		.BorderImage(&Style->BackgroundBrush)
@@ -133,7 +133,7 @@ FReply SDepthBar::OnMouseButtonDown(const FGeometry& MyGeometry, const FPointerE
 		TOptional<FBarPositions> OptionalPositions = GetBarPositions(MyGeometry); 
 		if (const FBarPositions* Positions = OptionalPositions.GetPtrOrNull())
 		{
-			const FVector2f MousePosition = MyGeometry.AbsoluteToLocal(MouseEvent.GetScreenSpacePosition());
+			const FVector2D MousePosition = MyGeometry.AbsoluteToLocal(MouseEvent.GetScreenSpacePosition());
 		
 			DragTarget = HoverTarget;
 			
@@ -141,10 +141,10 @@ FReply SDepthBar::OnMouseButtonDown(const FGeometry& MyGeometry, const FPointerE
 			{
 			case EDragTarget::Far_Handle:
 			case EDragTarget::Slice:
-				DragPosition = FVector2f(0.0f, Positions->BarFarPlane);
+				DragPosition = FVector2D(0.0, Positions->BarFarPlane);
 				break;
 			case EDragTarget::Near_Handle:
-				DragPosition = FVector2f(0.0f, Positions->BarNearPlane);
+				DragPosition = FVector2D(0.0, Positions->BarNearPlane);
 				break;
 			default:
 				DragPosition = MousePosition;
@@ -170,26 +170,27 @@ FReply SDepthBar::OnMouseMove(const FGeometry& MyGeometry, const FPointerEvent&
 	{
 		if (const FBarPositions* Positions = OptionalPositions.GetPtrOrNull())
 		{
-			float DragRatio = 1.0f;
+			double DragRatio = 1.0;
 			
 			if (MouseEvent.IsAltDown())
 			{
 				// Holding alt down should give consistent world-space movement (1px = 1u), but not if doing so would _lose_ precision.
-				DragRatio = FMath::Min(1.0f, (Positions->BarMax - Positions->BarMin) / static_cast<float>(Positions->DepthSpace.GetAlignedMax() - Positions->DepthSpace.GetAlignedMin()));
+				DragRatio = FMath::Min(1.0, (Positions->BarMax - Positions->BarMin) / (Positions->DepthSpace.GetAlignedMax() - Positions->DepthSpace.GetAlignedMin()));
 			}
 			
 			if (MouseEvent.IsControlDown())
 			{
-				DragRatio *= 0.1f;
+				DragRatio *= 0.1;
 			}
 			
-			DragPosition += MouseEvent.GetCursorDelta() * DragRatio;
+			// Use doubles for higher precision with very large worlds
+			DragPosition += FVector2D(MouseEvent.GetCursorDelta()) * DragRatio;
 		
 			if (Mode.Get() == EMode::Orthographic)
 			{
 				if (DragTarget == EDragTarget::Far_Handle)
 				{
-					const float ClampedBarPosition = FMath::Clamp(DragPosition.Y, Positions->BarMin, Positions->BarNearPlane - MinIndicatorDifference);
+					const double ClampedBarPosition = FMath::Clamp(DragPosition.Y, Positions->BarMin, Positions->BarNearPlane - MinIndicatorDifference);
 					
 					if (ClampedBarPosition == Positions->BarMin)
 					{
@@ -204,7 +205,7 @@ FReply SDepthBar::OnMouseMove(const FGeometry& MyGeometry, const FPointerEvent&
 				}
 				if (DragTarget == EDragTarget::Near_Handle)
 				{
-					const float ClampedBarPosition = FMath::Clamp(DragPosition.Y, Positions->BarFarPlane + MinIndicatorDifference, Positions->BarMax);
+					const double ClampedBarPosition = FMath::Clamp(DragPosition.Y, Positions->BarFarPlane + MinIndicatorDifference, Positions->BarMax);
 					if (ClampedBarPosition == Positions->BarMax)
 					{
 						OnNearPlaneChanged.ExecuteIfBound(TOptional<double>());
@@ -218,8 +219,7 @@ FReply SDepthBar::OnMouseMove(const FGeometry& MyGeometry, const FPointerEvent&
 				}
 				if (DragTarget == EDragTarget::Slice)
 				{
-					const float RawPosition = DragPosition.Y;
-					const float ClampedBarPosition = FMath::Clamp(RawPosition, Positions->BarMin, Positions->BarMax - (Positions->BarNearPlane - Positions->BarFarPlane));
+					const double ClampedBarPosition = FMath::Clamp(DragPosition.Y, Positions->BarMin, Positions->BarMax - (Positions->BarNearPlane - Positions->BarFarPlane));
 					
 					OnFarPlaneChanged.ExecuteIfBound(Positions->BarToAlignedPosition(ClampedBarPosition));
 					OnNearPlaneChanged.ExecuteIfBound(Positions->BarToAlignedPosition(ClampedBarPosition + (Positions->BarNearPlane - Positions->BarFarPlane)));
@@ -270,7 +270,7 @@ FReply SDepthBar::OnMouseButtonUp(const FGeometry& MyGeometry, const FPointerEve
 	if (DragTarget != EDragTarget::None)
 	{
 		constexpr float ClickThreshold = 3.0f;
-		if (FVector2f::DistSquared(DragPosition, StartDragPosition) < ClickThreshold)
+		if (FVector2D::DistSquared(DragPosition, StartDragPosition) < ClickThreshold)
 		{
 			if (DragTarget == EDragTarget::Far_Button)
 			{
@@ -334,8 +334,8 @@ void SDepthBar::OnArrangeChildren(const FGeometry& AllottedGeometry, FArrangedCh
 				const FVector Max = Sphere.Center + BarPositions->DepthSpace.GetForward() * Sphere.W;
 				const FVector Min = Sphere.Center - BarPositions->DepthSpace.GetForward() * Sphere.W;
 					
-				const float BarMax = BarPositions->WorldToBarPosition(Max);
-				const float BarMin = BarPositions->WorldToBarPosition(Min);
+				const float BarMax = static_cast<float>(BarPositions->WorldToBarPosition(Max));
+				const float BarMin = static_cast<float>(BarPositions->WorldToBarPosition(Min));
 				const float BarSize = FMath::Max(BarMin - BarMax, 1.0f);
 				
 				FVector2f Offset(Inset.Left, 0.0f);
@@ -493,14 +493,20 @@ int32 SDepthBar::OnPaint(const FPaintArgs& Args, const FGeometry& AllottedGeomet
 	TOptional<FBarPositions> OptionalPositions = GetBarPositions(AllottedGeometry); 
 	if (const FBarPositions* BarPositions = OptionalPositions.GetPtrOrNull())
 	{
+		// Truncate double values for slate
+		const float BarFarPlane = static_cast<float>(BarPositions->BarFarPlane);
+		const float BarNearPlane = static_cast<float>(BarPositions->BarNearPlane);
+		const float BarMin = static_cast<float>(BarPositions->BarMin);
+		const float BarMax = static_cast<float>(BarPositions->BarMax);
+	
 		// Draw the track
 		FSlateDrawElement::MakeBox(OutDrawElements,
 			LayerId++,
 			AllottedGeometry.MakeChild(
-				FVector2f(Style->TrackWidth, BarPositions->BarMax - BarPositions->BarMin + Style->SliceHandleSize.Y * 2.0f),
+				FVector2f(Style->TrackWidth, BarMax - BarMin + Style->SliceHandleSize.Y * 2.0f),
 				FSlateLayoutTransform(FVector2f(
 					Center.X - Style->TrackWidth * 0.5f,
-					BarPositions->BarMin - Style->SliceHandleSize.Y
+					BarMin - Style->SliceHandleSize.Y
 				))
 			).ToPaintGeometry(),
 			&Style->TrackBrush,
@@ -524,8 +530,8 @@ int32 SDepthBar::OnPaint(const FPaintArgs& Args, const FGeometry& AllottedGeomet
 		FSlateDrawElement::MakeBox(OutDrawElements,
 			LayerId,
 			AllottedGeometry.MakeChild(
-				FVector2f(Style->SliceWidth, BarPositions->BarNearPlane - BarPositions->BarFarPlane),
-				FSlateLayoutTransform(FVector2f(Center.X - Style->SliceWidth * 0.5f, BarPositions->BarFarPlane))
+				FVector2f(Style->SliceWidth, BarNearPlane - BarFarPlane),
+				FSlateLayoutTransform(FVector2f(Center.X - Style->SliceWidth * 0.5f, BarFarPlane))
 			).ToPaintGeometry(),
 			CurrentSliceBrush,
 			ESlateDrawEffect::None,
@@ -543,7 +549,7 @@ int32 SDepthBar::OnPaint(const FPaintArgs& Args, const FGeometry& AllottedGeomet
 			++LayerId,
 			AllottedGeometry.MakeChild(
 				FVector2f(Style->SliceHandleSize),
-				FSlateLayoutTransform(FVector2f(Center.X - Style->SliceHandleSize.X * 0.5f, BarPositions->BarFarPlane - Style->SliceHandleSize.Y))
+				FSlateLayoutTransform(FVector2f(Center.X - Style->SliceHandleSize.X * 0.5f, BarFarPlane - Style->SliceHandleSize.Y))
 			).ToPaintGeometry(),
 			TopBrush,
 			ESlateDrawEffect::None,
@@ -555,7 +561,7 @@ int32 SDepthBar::OnPaint(const FPaintArgs& Args, const FGeometry& AllottedGeomet
 			++LayerId,
 			AllottedGeometry.MakeChild(
 				FVector2f(Style->SliceHandleSize),
-				FSlateLayoutTransform(FVector2f(Center.X - Style->SliceHandleSize.X * 0.5f, BarPositions->BarNearPlane))
+				FSlateLayoutTransform(FVector2f(Center.X - Style->SliceHandleSize.X * 0.5f, BarNearPlane))
 			).ToPaintGeometry(),
 			BottomBrush,
 			ESlateDrawEffect::None,
@@ -563,13 +569,13 @@ int32 SDepthBar::OnPaint(const FPaintArgs& Args, const FGeometry& AllottedGeomet
 		);
 		
 		// Draw the "shades" over out-of-slice content
-		if (BarPositions->BarMin < BarPositions->BarFarPlane - Style->SliceHandleSize.Y)
+		if (BarMin < BarFarPlane - Style->SliceHandleSize.Y)
 		{
 			FSlateDrawElement::MakeBox(OutDrawElements,
 				LayerId,
 				AllottedGeometry.MakeChild(
-					FVector2f(AllottedGeometry.GetLocalSize().X, BarPositions->BarFarPlane - Style->SliceHandleSize.Y - BarPositions->BarMin),
-					FSlateLayoutTransform(FVector2f(0.0f, BarPositions->BarMin)
+					FVector2f(AllottedGeometry.GetLocalSize().X, BarFarPlane - Style->SliceHandleSize.Y - BarMin),
+					FSlateLayoutTransform(FVector2f(0.0f, BarMin)
 				)).ToPaintGeometry(),
 				&Style->BackgroundBrush,
 				ESlateDrawEffect::None,
@@ -577,13 +583,13 @@ int32 SDepthBar::OnPaint(const FPaintArgs& Args, const FGeometry& AllottedGeomet
 			);
 		}
 		
-		if (BarPositions->BarMax > BarPositions->BarNearPlane + Style->SliceHandleSize.Y)
+		if (BarMax > BarNearPlane + Style->SliceHandleSize.Y)
 		{
 			FSlateDrawElement::MakeBox(OutDrawElements,
 				LayerId,
 				AllottedGeometry.MakeChild(
-					FVector2f(AllottedGeometry.GetLocalSize().X, BarPositions->BarMax - BarPositions->BarNearPlane - Style->SliceHandleSize.Y),
-					FSlateLayoutTransform(FVector2f(0.0f, BarPositions->BarNearPlane + Style->SliceHandleSize.Y)
+					FVector2f(AllottedGeometry.GetLocalSize().X, BarMax - BarNearPlane - Style->SliceHandleSize.Y),
+					FSlateLayoutTransform(FVector2f(0.0f, BarNearPlane + Style->SliceHandleSize.Y)
 				)).ToPaintGeometry(),
 				&Style->BackgroundBrush,
 				ESlateDrawEffect::None,
@@ -709,22 +715,22 @@ double SDepthBar::FDepthSpace::BoundsToRelativePosition(double BoundsPosition) c
 	return AlignedToRelativePosition(Adjusted);
 }
 
-float SDepthBar::FBarPositions::WorldToBarPosition(const FVector& WorldPosition) const
+double SDepthBar::FBarPositions::WorldToBarPosition(const FVector& WorldPosition) const
 {
 	const double AxisPosition = DepthSpace.GetAlignedPosition(WorldPosition);
 	const double RelativePosition = DepthSpace.AlignedToRelativePosition(AxisPosition);
 	return RelativeToBarPosition(RelativePosition);
 }
 
-float SDepthBar::FBarPositions::RelativeToBarPosition(double RelativePosition) const
+double SDepthBar::FBarPositions::RelativeToBarPosition(double RelativePosition) const
 {
-	return FMath::Clamp(1.0f - static_cast<float>(RelativePosition), 0.0f, 1.0f) * (BarMax - BarMin) + BarMin;
+	return FMath::Clamp(1.0 - RelativePosition, 0.0, 1.0) * (BarMax - BarMin) + BarMin;
 }
 
-double SDepthBar::FBarPositions::BarToAlignedPosition(float BarPosition) const
+double SDepthBar::FBarPositions::BarToAlignedPosition(double BarPosition) const
 {
-	const float BarSize = BarMax - BarMin;
-	const float RelativePosition = BarSize > 0.0f ? (1.0f - (BarPosition - BarMin) / BarSize) : 0.0f;
+	const double BarSize = BarMax - BarMin;
+	const double RelativePosition = BarSize > 0.0 ? (1.0 - (BarPosition - BarMin) / BarSize) : 0.0;
 	return DepthSpace.RelativeToAlignedPosition(RelativePosition);
 }
 
diff --git a/Engine/Source/Editor/EditorFramework/Public/Viewports/SDepthBar.h b/Engine/Source/Editor/EditorFramework/Public/Viewports/SDepthBar.h
--- a/Engine/Source/Editor/EditorFramework/Public/Viewports/SDepthBar.h
+++ b/Engine/Source/Editor/EditorFramework/Public/Viewports/SDepthBar.h
@@ -161,16 +161,16 @@ private:
 		FDepthSpace DepthSpace;
 	
 		/** The Top of the bar (i.e. the far plane max) */
-		float BarMin;
+		double BarMin;
 		/** The Bottom of the bar (i.e. the near plane min) */
-		float BarMax;
+		double BarMax;
 		
-		float BarFarPlane;
-		float BarNearPlane;
+		double BarFarPlane;
+		double BarNearPlane;
 		
-		float WorldToBarPosition(const FVector& WorldPosition) const;
-		float RelativeToBarPosition(double RelativePosition) const;
-		double BarToAlignedPosition(float BarPosition) const;
+		double WorldToBarPosition(const FVector& WorldPosition) const;
+		double RelativeToBarPosition(double RelativePosition) const;
+		double BarToAlignedPosition(double BarPosition) const;
 	};
 
 	DECLARE_DELEGATE_RetVal_OneParam(float, FGetBarPosition, const FBarPositions& BarPositions)
@@ -251,8 +251,10 @@ private:
 	};
 	EDragTarget HoverTarget = EDragTarget::None;
 	EDragTarget DragTarget = EDragTarget::None;
-	FVector2f StartDragPosition = FVector2f::Zero();
-	FVector2f DragPosition = FVector2f::Zero();
+	
+	// Store drag position as doubles rather than floats to handle precision drags on very large worlds
+	FVector2D StartDragPosition = FVector2D::Zero();
+	FVector2D DragPosition = FVector2D::Zero();
 	
 	TAttribute<EMode> Mode;
 	TAttribute<TOptional<FDepthSpace>> DepthSpace;
