[Gizmos] Update various scale sources to provide relative scale rather than world scale to match the relative-only nature of scale manipulations. Fix issues with scaling child actors of scaled parents.

#rnx
#Synced-CL 52300119

diff --git a/Engine/Source/Editor/EditorFramework/Private/Elements/Framework/TypedElementViewportInteraction.cpp b/Engine/Source/Editor/EditorFramework/Private/Elements/Framework/TypedElementViewportInteraction.cpp
--- a/Engine/Source/Editor/EditorFramework/Private/Elements/Framework/TypedElementViewportInteraction.cpp
+++ b/Engine/Source/Editor/EditorFramework/Private/Elements/Framework/TypedElementViewportInteraction.cpp
@@ -169,10 +169,10 @@ bool FTypedElementViewportInteractionCustomization::GetFocusBounds(const TTypedE
 
 bool FTypedElementViewportInteractionCustomization::GetScaleForGizmo(const TTypedElement<ITypedElementWorldInterface>& InElementWorldHandle, FVector& OutScale)
 {
-	FTransform WorldTransform;
-	if (InElementWorldHandle.GetWorldTransform(WorldTransform))
+	FTransform Transform;
+	if (InElementWorldHandle.GetRelativeTransform(Transform))
 	{
-		OutScale = WorldTransform.GetScale3D();
+		OutScale = Transform.GetScale3D();
 		return true;
 	}
 
diff --git a/Engine/Source/Editor/UnrealEd/Private/Elements/Component/ComponentElementEditorViewportInteractionCustomization.cpp b/Engine/Source/Editor/UnrealEd/Private/Elements/Component/ComponentElementEditorViewportInteractionCustomization.cpp
--- a/Engine/Source/Editor/UnrealEd/Private/Elements/Component/ComponentElementEditorViewportInteractionCustomization.cpp
+++ b/Engine/Source/Editor/UnrealEd/Private/Elements/Component/ComponentElementEditorViewportInteractionCustomization.cpp
@@ -155,7 +155,7 @@ bool FComponentElementEditorViewportInteractionCustomization::GetScaleForGizmo(c
 				return false;
 			}
 
-			OutScale = PrimitiveComponent->GetComponentScale();
+			OutScale = PrimitiveComponent->GetRelativeScale3D();
 			return true;
 		}
 	}
diff --git a/Engine/Source/Editor/UnrealEd/Private/LevelEditorViewport.cpp b/Engine/Source/Editor/UnrealEd/Private/LevelEditorViewport.cpp
--- a/Engine/Source/Editor/UnrealEd/Private/LevelEditorViewport.cpp
+++ b/Engine/Source/Editor/UnrealEd/Private/LevelEditorViewport.cpp
@@ -4626,7 +4626,12 @@ FTransform FLevelEditorViewportClient::GetPreDragActorTransform(const AActor* In
 	{
 		return *PreDragTransform;
 	}
-	return InActor->GetTransform();
+	
+	FTransform Transform = InActor->GetTransform();
+	// Scale manipulation and delta calculation are always based off of the relative scale.
+	// Provide that information here.
+	Transform.SetScale3D(InActor->GetActorRelativeScale3D());
+	return Transform;
 }
 
 EMouseCursor::Type FLevelEditorViewportClient::GetCursor(FViewport* InViewport,int32 X,int32 Y)
